Copying //UE4/Partner-MagicLeap-Staging to Dev-VR (//UE4/Dev-VR)

diff --git a/Engine/Plugins/Lumin/MagicLeapCamera/Source/MagicLeapCameraPreview/Private/MagicLeapCameraPreviewPlayer.h b/Engine/Plugins/Lumin/MagicLeapCamera/Source/MagicLeapCameraPreview/Private/MagicLeapCameraPreviewPlayer.h
new file mode 100644
--- /dev/null
+++ b/Engine/Plugins/Lumin/MagicLeapCamera/Source/MagicLeapCameraPreview/Private/MagicLeapCameraPreviewPlayer.h
@@ -0,0 +1,130 @@
+// Copyright 1998-2019 Epic Games, Inc. All Rights Reserved.
+
+#pragma once
+#include "CoreMinimal.h"
+#include "Containers/UnrealString.h"
+#include "Misc/Guid.h"
+#include "IMediaCache.h"
+#include "IMediaControls.h"
+#include "IMediaPlayer.h"
+#include "IMediaTracks.h"
+#include "IMediaView.h"
+#include "Misc/Timespan.h"
+#include "Templates/SharedPointer.h"
+#include "Templates/Atomic.h"
+#include "RenderingThread.h"
+#include "RHI.h"
+#include "RHIResources.h"
+#include "MagicLeapPluginUtil.h" // for ML_INCLUDES_START/END
+
+class FMediaSamples;
+class IMediaEventSink;
+struct FMagicLeapVideoTextureData;
+
+/**
+ *	Implement media playback using the MagicLeap MediaPlayer interface.
+ */
+class FMagicLeapCameraPreviewPlayer : public IMediaPlayer, public IMediaControls, public IMediaCache, public IMediaTracks, public IMediaView
+{
+public:
+
+	/**
+	 * Create and initialize a new instance.
+	 *
+	 * @param InEventSink The object that receives media events from this player.
+	 */
+	FMagicLeapCameraPreviewPlayer(IMediaEventSink& InEventSink);
+
+	/** Virtual destructor. */
+	virtual ~FMagicLeapCameraPreviewPlayer();
+
+public:
+	/** IMediaPlayer interface */
+	virtual void Close() override;
+	virtual IMediaCache& GetCache() override;
+	virtual IMediaControls& GetControls() override;
+	virtual FString GetInfo() const override;
+	virtual FName GetPlayerName() const override;
+	virtual IMediaSamples& GetSamples() override;
+	virtual FString GetStats() const override;
+	virtual IMediaTracks& GetTracks() override;
+	virtual FString GetUrl() const override;
+	virtual IMediaView& GetView() override;
+	virtual bool Open(const FString& Url, const IMediaOptions* Options) override;
+	virtual bool Open(const TSharedRef<FArchive, ESPMode::ThreadSafe>& Archive, const FString& OriginalUrl, const IMediaOptions* Options) override;
+	virtual void SetGuid(const FGuid& Guid) override;
+	virtual void TickFetch(FTimespan DeltaTime, FTimespan Timecode) override;
+	virtual void TickInput(FTimespan DeltaTime, FTimespan Timecode) override;
+
+public:
+	/** IMediaControls interface */
+	virtual bool CanControl(EMediaControl Control) const override;
+	virtual FTimespan GetDuration() const override;
+	virtual float GetRate() const override;
+	virtual EMediaState GetState() const override;
+	virtual EMediaStatus GetStatus() const override;
+	virtual TRangeSet<float> GetSupportedRates(EMediaRateThinning Thinning) const override;
+	virtual FTimespan GetTime() const override;
+	virtual bool IsLooping() const override;
+	virtual bool Seek(const FTimespan& Time) override;
+	virtual bool SetLooping(bool Looping) override;
+	virtual bool SetRate(float Rate) override;
+	virtual bool SetNativeVolume(float Volume) override;
+
+public:
+	/** IMediaTracks interface */
+	virtual bool GetAudioTrackFormat(int32 TrackIndex, int32 FormatIndex, FMediaAudioTrackFormat& OutFormat) const override;
+	virtual int32 GetNumTracks(EMediaTrackType TrackType) const override;
+	virtual int32 GetNumTrackFormats(EMediaTrackType TrackType, int32 TrackIndex) const override;
+	virtual int32 GetSelectedTrack(EMediaTrackType TrackType) const override;
+	virtual FText GetTrackDisplayName(EMediaTrackType TrackType, int32 TrackIndex) const override;
+	virtual int32 GetTrackFormat(EMediaTrackType TrackType, int32 TrackIndex) const override;
+	virtual FString GetTrackLanguage(EMediaTrackType TrackType, int32 TrackIndex) const override;
+	virtual FString GetTrackName(EMediaTrackType TrackType, int32 TrackIndex) const override;
+	virtual bool GetVideoTrackFormat(int32 TrackIndex, int32 FormatIndex, FMediaVideoTrackFormat& OutFormat) const override;
+	virtual bool SelectTrack(EMediaTrackType TrackType, int32 TrackIndex) override;
+	virtual bool SetTrackFormat(EMediaTrackType TrackType, int32 TrackIndex, int32 FormatIndex) override;
+
+private:
+	/** Callback for when the application resumed in the foreground. */
+	void HandleApplicationHasEnteredForeground();
+
+	/** Callback for when the application is being paused in the background. */
+	void HandleApplicationWillEnterBackground();
+
+protected:
+	bool bMediaPrepared;
+
+	/** Current player state. */
+	EMediaState CurrentState;
+
+	/** The media event handler. */
+	IMediaEventSink& EventSink;
+
+	/** Media player Guid */
+	FGuid PlayerGuid;
+
+	/** Foreground/background delegate for pause. */
+	FDelegateHandle PauseHandle;
+
+	/** Foreground/background delegate for resume. */
+	FDelegateHandle ResumeHandle;
+
+	/** The media sample queue. */
+	TSharedPtr<FMediaSamples, ESPMode::ThreadSafe> Samples;
+
+	TSharedPtr<FMagicLeapVideoTextureData, ESPMode::ThreadSafe> TextureData;
+
+	TMap<EMediaTrackType, TArray<int32>> TrackInfo;
+
+	TMap<EMediaTrackType, int32> SelectedTrack;
+
+	FCriticalSection CriticalSection;
+
+	bool bWasMediaPlayingBeforeAppPause;
+
+	bool bPlaybackCompleted;
+
+private:
+	virtual void RegisterExternalTexture(const FGuid& InGuid, FTextureRHIRef& InTextureRHI, FSamplerStateRHIRef& InSamplerStateRHI);
+};
diff --git a/Engine/Plugins/Lumin/MagicLeapCamera/Source/MagicLeapCameraPreview/Public/IMagicLeapCameraPreviewModule.h b/Engine/Plugins/Lumin/MagicLeapCamera/Source/MagicLeapCameraPreview/Public/IMagicLeapCameraPreviewModule.h
new file mode 100644
--- /dev/null
+++ b/Engine/Plugins/Lumin/MagicLeapCamera/Source/MagicLeapCameraPreview/Public/IMagicLeapCameraPreviewModule.h
@@ -0,0 +1,44 @@
+// Copyright 1998-2019 Epic Games, Inc. All Rights Reserved.
+
+#pragma once
+
+#include "Modules/ModuleManager.h"
+#include "Containers/Ticker.h"
+
+class IMediaPlayer;
+
+/**
+ * The public interface to this module.  In most cases, this interface is only public to sibling modules
+ * within this plugin.
+ */
+class IMagicLeapCameraPreviewModule : public IModuleInterface
+{
+public:
+	/**
+	 * Singleton-like access to this module's interface.  This is just for convenience!
+	 * Beware of calling this during the shutdown phase, though.  Your module might have been unloaded already.
+	 *
+	 * @return Returns singleton instance, loading the module on demand if needed
+	 */
+	static inline IMagicLeapCameraPreviewModule& Get()
+	{
+		return FModuleManager::LoadModuleChecked<IMagicLeapCameraPreviewModule>("MagicLeapCameraPreview");
+	}
+
+	/**
+	 * Checks to see if this module is loaded and ready.  It is only valid to call Get() if IsAvailable() returns true.
+	 *
+	 * @return True if the module is loaded and ready to use
+	 */
+	static inline bool IsAvailable()
+	{
+		return FModuleManager::Get().IsModuleLoaded("MagicLeapCameraPreview");
+	}
+
+	/**
+	 * Creates a camera preview media player.
+	 *
+	 * @return A new media player, or nullptr if a player couldn't be created.
+	 */
+	virtual TSharedPtr<IMediaPlayer, ESPMode::ThreadSafe> CreatePreviewPlayer(class IMediaEventSink& EventSink) = 0;
+};
diff --git a/Engine/Plugins/Lumin/MagicLeapCamera/Source/MagicLeapCameraPreviewFactory/MagicLeapCameraPreviewFactory.Build.cs b/Engine/Plugins/Lumin/MagicLeapCamera/Source/MagicLeapCameraPreviewFactory/MagicLeapCameraPreviewFactory.Build.cs
new file mode 100644
--- /dev/null
+++ b/Engine/Plugins/Lumin/MagicLeapCamera/Source/MagicLeapCameraPreviewFactory/MagicLeapCameraPreviewFactory.Build.cs
@@ -0,0 +1,32 @@
+// Copyright 1998-2019 Epic Games, Inc. All Rights Reserved.
+
+namespace UnrealBuildTool.Rules
+{
+	public class MagicLeapCameraPreviewFactory : ModuleRules
+	{
+		public MagicLeapCameraPreviewFactory(ReadOnlyTargetRules Target) : base(Target)
+		{
+			DynamicallyLoadedModuleNames.Add("Media");
+
+			PrivateDependencyModuleNames.AddRange(
+				new string[] {
+					"Core",
+					"CoreUObject",
+					"MediaAssets"
+				});
+
+			PrivateIncludePathModuleNames.AddRange(
+				new string[] {
+					"MagicLeapCameraPreview",
+					"Media"
+				});
+
+			PrivateIncludePaths.Add("MagicLeapCameraPreviewFactory/Private");
+
+			if (Target.Platform == UnrealTargetPlatform.Lumin)
+			{
+				DynamicallyLoadedModuleNames.Add("MagicLeapCameraPreview");
+			}
+		}
+	}
+}
diff --git a/Engine/Plugins/Lumin/MagicLeapCamera/Source/MagicLeapCameraPreviewFactory/Private/MagicLeapCameraPreviewFactoryModule.cpp b/Engine/Plugins/Lumin/MagicLeapCamera/Source/MagicLeapCameraPreviewFactory/Private/MagicLeapCameraPreviewFactoryModule.cpp
new file mode 100644
--- /dev/null
+++ b/Engine/Plugins/Lumin/MagicLeapCamera/Source/MagicLeapCameraPreviewFactory/Private/MagicLeapCameraPreviewFactoryModule.cpp
@@ -0,0 +1,92 @@
+// Copyright 1998-2019 Epic Games, Inc. All Rights Reserved.
+
+#include "CoreMinimal.h"
+#include "Misc/Paths.h"
+#include "Modules/ModuleInterface.h"
+#include "Modules/ModuleManager.h"
+#include "IMediaModule.h"
+#include "IMediaPlayerFactory.h"
+#include "Misc/ConfigCacheIni.h"
+#include "IMagicLeapCameraPreviewModule.h"
+
+#define LOCTEXT_NAMESPACE "FMagicLeapCameraPreviewFactoryModule"
+
+/**
+ * Implements the MagicLeapCameraPreviewFactory module.
+ */
+class FMagicLeapCameraPreviewFactoryModule : public IMediaPlayerFactory, public IModuleInterface
+{
+public:
+	/** IMediaPlayerFactory interface */
+	virtual bool CanPlayUrl(const FString& Url, const IMediaOptions* /*Options*/, TArray<FText>* /*OutWarnings*/, TArray<FText>* OutErrors) const override
+	{
+		(void)Url;
+		(void)OutErrors;
+		return true;
+	}
+
+	virtual TSharedPtr<IMediaPlayer, ESPMode::ThreadSafe> CreatePlayer(IMediaEventSink& EventSink) override
+	{
+		return IMagicLeapCameraPreviewModule::Get().CreatePreviewPlayer(EventSink);
+	}
+
+	virtual FText GetDisplayName() const override
+	{
+		return LOCTEXT("MediaPlayerDisplayName", "MagicLeap Camera Preview");
+	}
+
+	virtual FName GetPlayerName() const override
+	{
+		static FName PlayerName(TEXT("MagicLeapCameraPreview"));
+		return PlayerName;
+	}
+
+	virtual const TArray<FString>& GetSupportedPlatforms() const override
+	{
+		return SupportedPlatforms;
+	}
+
+	virtual bool SupportsFeature(EMediaFeature Feature) const
+	{
+		return ((Feature == EMediaFeature::AudioTracks) ||
+			(Feature == EMediaFeature::VideoSamples) ||
+			(Feature == EMediaFeature::VideoTracks));
+	}
+
+public:
+	/** IModuleInterface interface */
+	
+	virtual void StartupModule() override
+	{
+    	// supported platforms
+    	SupportedPlatforms.Add(TEXT("Lumin"));
+    	// Hack until we get a separete ini platform for Lumin. Will not affect Android since this plugin is not built for it.
+    	SupportedPlatforms.Add(TEXT("Android"));
+		// register media player info
+		auto MediaModule = FModuleManager::LoadModulePtr<IMediaModule>("Media");
+
+		if (MediaModule != nullptr)
+		{
+			MediaModule->RegisterPlayerFactory(*this);
+		}
+	}
+
+	virtual void ShutdownModule() override
+	{
+		// unregister player factory
+		auto MediaModule = FModuleManager::GetModulePtr<IMediaModule>("Media");
+		
+		if (MediaModule != nullptr)
+		{
+			MediaModule->UnregisterPlayerFactory(*this);
+		}
+	}
+
+private:
+	/** List of platforms that the media player support. */
+	TArray<FString> SupportedPlatforms;
+};
+
+#undef LOCTEXT_NAMESPACE
+
+IMPLEMENT_MODULE(FMagicLeapCameraPreviewFactoryModule, MagicLeapCameraPreviewFactory);
diff --git a/Engine/Plugins/Lumin/MagicLeapContacts/MagicLeapContacts.uplugin b/Engine/Plugins/Lumin/MagicLeapContacts/MagicLeapContacts.uplugin
new file mode 100644
--- /dev/null
+++ b/Engine/Plugins/Lumin/MagicLeapContacts/MagicLeapContacts.uplugin
@@ -0,0 +1,44 @@
+{
+	"FileVersion": 3,
+	"Version": 1,
+	"VersionName": "1.0",
+	"FriendlyName": "Magic Leap Contacts",
+	"CreatedBy": "Magic Leap, Inc.",
+	"CreatedByURL": "http://magicleap.com",
+	"Description": "Support for the Magic Leap Contacts API.",
+	"Category": "Undefined",
+	"EnabledByDefault": false,
+	"DocsURL": "",
+	"MarketplaceURL": "",
+	"SupportURL": "",
+	"CanContainContent": false,
+	"IsBetaVersion": false,
+	"Installed": false,
+	"SupportedTargetPlatforms": [
+		"Lumin",
+		"Win64",
+		"Mac"
+	],
+	"Modules": [
+		{
+			"Name": "MagicLeapContacts",
+			"Type": "Runtime",
+			"LoadingPhase": "PostConfigInit",
+			"WhitelistPlatforms": [
+				"Win64",
+				"Mac",
+				"Lumin"
+			]
+		}
+	],
+	"Plugins": [
+		{
+			"Name": "MLSDK",
+			"Enabled": true
+		},
+		{
+			"Name": "MagicLeap",
+			"Enabled": true
+		}
+	]
+}
diff --git a/Engine/Plugins/Lumin/MagicLeapContacts/Source/MagicLeapContacts.Build.cs b/Engine/Plugins/Lumin/MagicLeapContacts/Source/MagicLeapContacts.Build.cs
new file mode 100644
--- /dev/null
+++ b/Engine/Plugins/Lumin/MagicLeapContacts/Source/MagicLeapContacts.Build.cs
@@ -0,0 +1,32 @@
+// Copyright 1998-2019 Epic Games, Inc. All Rights Reserved.
+
+using System.IO;
+
+namespace UnrealBuildTool.Rules
+{
+	public class MagicLeapContacts : ModuleRules
+	{
+		public MagicLeapContacts( ReadOnlyTargetRules Target ) : base(Target)
+		{
+			PublicIncludePaths.Add(Path.Combine(ModuleDirectory, "Public"));
+
+			PublicDependencyModuleNames.AddRange(
+				new string[]
+				{
+					"Core",
+					"CoreUObject",
+					"Engine"
+				}
+			);
+
+			PrivateDependencyModuleNames.AddRange(
+				new string[]
+				{
+					"LuminRuntimeSettings",
+					"MLSDK",
+					"MagicLeap"
+				}
+			);
+		}
+	}
+}
diff --git a/Engine/Plugins/Lumin/MagicLeapContacts/Source/Private/MagicLeapContactsComponent.cpp b/Engine/Plugins/Lumin/MagicLeapContacts/Source/Private/MagicLeapContactsComponent.cpp
new file mode 100644
--- /dev/null
+++ b/Engine/Plugins/Lumin/MagicLeapContacts/Source/Private/MagicLeapContactsComponent.cpp
@@ -0,0 +1,47 @@
+// Copyright 1998-2019 Epic Games, Inc. All Rights Reserved.
+
+#include "MagicLeapContactsComponent.h"
+#include "MagicLeapContactsPlugin.h"
+
+void UMagicLeapContactsComponent::BeginPlay()
+{
+	Super::BeginPlay();
+	GET_MAGIC_LEAP_CONTACTS_PLUGIN()->Startup();
+	GET_MAGIC_LEAP_CONTACTS_PLUGIN()->SetLogDelegate(OnLogMessage);
+}
+
+void UMagicLeapContactsComponent::EndPlay(const EEndPlayReason::Type EndPlayReason)
+{
+	GET_MAGIC_LEAP_CONTACTS_PLUGIN()->Shutdown();
+	Super::EndPlay(EndPlayReason);
+}
+
+FGuid UMagicLeapContactsComponent::AddContactAsync(const FMagicLeapContact& Contact)
+{
+	return GET_MAGIC_LEAP_CONTACTS_PLUGIN()->AddContactAsync(Contact, OnAddContactResult);
+}
+
+FGuid UMagicLeapContactsComponent::EditContactAsync(const FMagicLeapContact& Contact)
+{
+	return GET_MAGIC_LEAP_CONTACTS_PLUGIN()->EditContactAsync(Contact, OnEditContactResult);
+}
+
+FGuid UMagicLeapContactsComponent::DeleteContactAsync(const FMagicLeapContact& Contact)
+{
+	return GET_MAGIC_LEAP_CONTACTS_PLUGIN()->DeleteContactAsync(Contact, OnDeleteContactResult);
+}
+
+FGuid UMagicLeapContactsComponent::RequestContactsAsync()
+{
+	return GET_MAGIC_LEAP_CONTACTS_PLUGIN()->RequestContactsAsync(OnRequestContactsResult);
+}
+
+FGuid UMagicLeapContactsComponent::SearchContactsAsync(const FString& Query, EMagicLeapContactsSearchField SearchField)
+{
+	return GET_MAGIC_LEAP_CONTACTS_PLUGIN()->SearchContactsAsync(Query, SearchField, OnSearchContactsResult);
+}
+
+bool UMagicLeapContactsComponent::CancelRequest(const FGuid& RequestHandle)
+{
+	return GET_MAGIC_LEAP_CONTACTS_PLUGIN()->CancelRequest(RequestHandle);
+}
diff --git a/Engine/Plugins/Lumin/MagicLeapContacts/Source/Private/MagicLeapContactsFunctionLibrary.cpp b/Engine/Plugins/Lumin/MagicLeapContacts/Source/Private/MagicLeapContactsFunctionLibrary.cpp
new file mode 100644
--- /dev/null
+++ b/Engine/Plugins/Lumin/MagicLeapContacts/Source/Private/MagicLeapContactsFunctionLibrary.cpp
@@ -0,0 +1,61 @@
+// Copyright 1998-2019 Epic Games, Inc. All Rights Reserved.
+
+#include "MagicLeapContactsFunctionLibrary.h"
+#include "MagicLeapContactsPlugin.h"
+
+bool UMagicLeapContactsFunctionLibrary::Startup()
+{
+	return GET_MAGIC_LEAP_CONTACTS_PLUGIN()->Startup();
+}
+
+bool UMagicLeapContactsFunctionLibrary::Shutdown()
+{
+	return GET_MAGIC_LEAP_CONTACTS_PLUGIN()->Shutdown();
+}
+
+FGuid UMagicLeapContactsFunctionLibrary::AddContactAsync(const FMagicLeapContact& Contact, const FSingleContactResultDelegate& InResultDelegate)
+{
+	FSingleContactResultDelegateMulti ResultDelegate;
+	ResultDelegate.Add(InResultDelegate);
+	return GET_MAGIC_LEAP_CONTACTS_PLUGIN()->AddContactAsync(Contact, ResultDelegate);
+}
+
+FGuid UMagicLeapContactsFunctionLibrary::EditContactAsync(const FMagicLeapContact& Contact, const FSingleContactResultDelegate& InResultDelegate)
+{
+	FSingleContactResultDelegateMulti ResultDelegate;
+	ResultDelegate.Add(InResultDelegate);
+	return GET_MAGIC_LEAP_CONTACTS_PLUGIN()->EditContactAsync(Contact, ResultDelegate);
+}
+
+FGuid UMagicLeapContactsFunctionLibrary::DeleteContactAsync(const FMagicLeapContact& Contact, const FSingleContactResultDelegate& InResultDelegate)
+{
+	FSingleContactResultDelegateMulti ResultDelegate;
+	ResultDelegate.Add(InResultDelegate);
+	return GET_MAGIC_LEAP_CONTACTS_PLUGIN()->DeleteContactAsync(Contact, ResultDelegate);
+}
+
+FGuid UMagicLeapContactsFunctionLibrary::RequestContactsAsync(const FMultipleContactsResultDelegate& InResultDelegate)
+{
+	FMultipleContactsResultDelegateMulti ResultDelegate;
+	ResultDelegate.Add(InResultDelegate);
+	return GET_MAGIC_LEAP_CONTACTS_PLUGIN()->RequestContactsAsync(ResultDelegate);
+}
+
+FGuid UMagicLeapContactsFunctionLibrary::SearchContactsAsync(const FString& Query, EMagicLeapContactsSearchField SearchField, const FMultipleContactsResultDelegate& InResultDelegate)
+{
+	FMultipleContactsResultDelegateMulti ResultDelegate;
+	ResultDelegate.Add(InResultDelegate);
+	return GET_MAGIC_LEAP_CONTACTS_PLUGIN()->SearchContactsAsync(Query, SearchField, ResultDelegate);
+}
+
+bool UMagicLeapContactsFunctionLibrary::CancelRequest(const FGuid& RequestHandle)
+{
+	return GET_MAGIC_LEAP_CONTACTS_PLUGIN()->CancelRequest(RequestHandle);
+}
+
+bool UMagicLeapContactsFunctionLibrary::SetLogDelegate(const FContactsLogMessage& InLogDelegate)
+{
+	FContactsLogMessageMulti LogDelegate;
+	LogDelegate.Add(InLogDelegate);
+	return GET_MAGIC_LEAP_CONTACTS_PLUGIN()->SetLogDelegate(LogDelegate);
+}
diff --git a/Engine/Plugins/Lumin/MagicLeapContacts/Source/Private/MagicLeapContactsPlugin.cpp b/Engine/Plugins/Lumin/MagicLeapContacts/Source/Private/MagicLeapContactsPlugin.cpp
new file mode 100644
--- /dev/null
+++ b/Engine/Plugins/Lumin/MagicLeapContacts/Source/Private/MagicLeapContactsPlugin.cpp
@@ -0,0 +1,730 @@
+// Copyright 1998-2019 Epic Games, Inc. All Rights Reserved.
+
+#include "MagicLeapContactsPlugin.h"
+
+using namespace MagicLeap;
+
+DEFINE_LOG_CATEGORY(LogMagicLeapContacts);
+
+FMagicLeapContactsPlugin::FMagicLeapContactsPlugin()
+: PhoneNumberPattern(TEXT("(^\\s*(?:\\+?(\\d{1,3}))?[-. (]*(\\d{3})[-. )]*(\\d{3})[-. ]*(\\d{4})(?: *x(\\d+))?\\s*$)"))
+, EmailPattern(TEXT("((?:[a-z0-9!#$%&'*+/=?^_`{|}~-]+(?:\\.[a-z0-9!#$%&'*+/=?^_`{|}~-]+)*|""(?:[\x01-\x08\x0b\x0c\x0e-\x1f\x21\x23-\x5b\x5d-\x7f]|\\[\x01-\x09\x0b\x0c\x0e-\x7f])*"")@(?:(?:[a-z0-9](?:[a-z0-9-]*[a-z0-9])?\\.)+[a-z0-9](?:[a-z0-9-]*[a-z0-9])?|\\[(?:(?:25[0-5]|2[0-4][0-9]|[01]?[0-9][0-9]?)\\.){3}(?:25[0-5]|2[0-4][0-9]|[01]?[0-9][0-9]?|[a-z0-9-]*[a-z0-9]:(?:[\x01-\x08\x0b\x0c\x0e-\x1f\x21-\x5a\x53-\x7f]|\\[\x01-\x09\x0b\x0c\x0e-\x7f])+)\\]))"))
+, MAX_NAME_LENGTH(128)
+, MAX_PHONE_LENGTH(20)
+, MAX_EMAIL_LENGTH(128)
+, MAX_TAG_LENGTH(128)
+, MAX_PHONE_COUNT(16)
+, MAX_EMAIL_COUNT(32)
+, PrivilegesManager({ EMagicLeapPrivilege::AddressBookRead, EMagicLeapPrivilege::AddressBookWrite })
+{
+
+}
+
+void FMagicLeapContactsPlugin::StartupModule()
+{
+	IMagicLeapContactsPlugin::StartupModule();
+	TickDelegate = FTickerDelegate::CreateRaw(this, &FMagicLeapContactsPlugin::Tick);
+	TickDelegateHandle = FTicker::GetCoreTicker().AddTicker(TickDelegate);
+}
+
+void FMagicLeapContactsPlugin::ShutdownModule()
+{
+	FTicker::GetCoreTicker().RemoveTicker(TickDelegateHandle);
+	IModuleInterface::ShutdownModule();
+}
+
+bool FMagicLeapContactsPlugin::Tick(float DeltaTime)
+{
+#if WITH_MLSDK
+	if (PendingRequests.Num() > 0)
+	{
+		ProcessPendingRequests();
+	}
+
+	if (ActiveRequests.Num() == 0)
+	{
+		return true;
+	}
+
+	MLResult Result = MLResult_Ok;
+	for (int32 i = ActiveRequests.Num() - 1; i > -1; --i)
+	{
+		bool bNeedsRemoving = false;
+		const FContactRequest& ContactRequest = ActiveRequests[i];
+		if (ContactRequest.SingleContactResultDelegate.IsBound())
+		{
+			MLContactsOperationResult* OpResult = nullptr;
+			Result = MLContactsTryGetOperationResult(ActiveRequests[i].Handle, &OpResult);
+			if (Result == MLResult_Pending)
+			{
+				continue;
+			}
+			else if (Result == MLContactsResult_Completed || Result == MLContactsResult_Cancelled)
+			{
+				ContactRequest.SingleContactResultDelegate.Broadcast(MLOpStatusToUEOpStatus(OpResult->operation_status));
+			}
+			else
+			{
+				Log(EMagicLeapContactsOperationStatus::Fail, FString::Printf(TEXT("MLContactsTryGetOperationResult failed with error '%s'!"), UTF8_TO_TCHAR(MLContactsGetResultString(Result))));
+			}
+
+			bNeedsRemoving = true;
+		}
+		else
+		{
+			MLContactsListResult* ListResult = nullptr;
+			Result = MLContactsTryGetListResult(ContactRequest.Handle, &ListResult);
+			if (Result == MLResult_Pending)
+			{
+				continue;
+			}
+			else if (Result == MLContactsResult_Completed)
+			{
+				const int32 NumContacts = static_cast<int32>(ListResult->result_list.count);
+				TArray<FMagicLeapContact> Contacts;
+				Contacts.AddDefaulted(NumContacts);
+				for (int32 iContact = 0; iContact < NumContacts; ++iContact)
+				{
+					FMagicLeapContact& Contact = Contacts[iContact];
+					MLToUEContact(ListResult->result_list.contacts[iContact], Contact);
+				}
+
+				ContactRequest.MultipleContactsResultDelegate.Broadcast(Contacts, MLOpStatusToUEOpStatus(ListResult->operation_status));
+			}
+			else
+			{
+				Log(EMagicLeapContactsOperationStatus::Fail, FString::Printf(TEXT("MLContactsTryGetOperationResult failed with error '%s'!"), UTF8_TO_TCHAR(MLContactsGetResultString(Result))));
+			}
+
+			bNeedsRemoving = true;
+		}
+
+		if (bNeedsRemoving)
+		{
+			MLContactsReleaseRequestResources(ContactRequest.Handle);
+			ActiveRequests.RemoveAt(i);
+		}
+	}
+#endif // WITH_MLSDK
+	return true;
+}
+
+bool FMagicLeapContactsPlugin::Startup()
+{
+#if WITH_MLSDK
+	MLResult Result = MLContactsStartup();
+	if (MLResult_Ok != Result)
+	{
+		Log(EMagicLeapContactsOperationStatus::Fail, FString::Printf(TEXT("MLContactsStartup failed with error '%s'"), UTF8_TO_TCHAR(MLContactsGetResultString(Result))));
+	}
+
+	return Result == MLResult_Ok;
+#else
+	return false;
+#endif // WITH_MLSDK
+}
+
+bool FMagicLeapContactsPlugin::Shutdown()
+{
+#if WITH_MLSDK
+	MLResult Result = MLContactsShutdown();
+	if (MLResult_Ok != Result)
+	{
+		Log(EMagicLeapContactsOperationStatus::Fail, FString::Printf(TEXT("MLContactsShutdown failed with error '%s'"), UTF8_TO_TCHAR(MLContactsGetResultString(Result))));
+	}
+	return Result == MLResult_Ok;
+#else
+	return false;
+#endif // WITH_MLSDK
+}
+
+FGuid FMagicLeapContactsPlugin::AddContactAsync(const FMagicLeapContact& InContact, const FSingleContactResultDelegateMulti& ResultDelegate)
+{
+#if WITH_MLSDK
+	if (TryAddPendingTask(EMagicLeapPrivilege::AddressBookWrite, FContactRequest::EType::Add, InContact, ResultDelegate))
+	{
+		return FGuid();
+	}
+
+	FMagicLeapContact UEContact = InContact;
+	ForceEmailAddressessToLower(UEContact);
+
+	if (!ValidateUEContact(UEContact))
+	{
+		if (ResultDelegate.IsBound())
+		{
+			ResultDelegate.Broadcast(EMagicLeapContactsOperationStatus::Fail);
+		}
+		return FGuid();
+	}
+
+	MLContactsContact MLContact;
+	UEToMLContact(UEContact, MLContact);
+	MLHandle RequestHandle;
+	MLResult Result = MLContactsRequestInsert(&MLContact, &RequestHandle);
+	DestroyMLContact(MLContact);
+	if (Result == MLResult_Ok)
+	{
+		FContactRequest& ContactRequest = ActiveRequests.AddDefaulted_GetRef();
+		ContactRequest.Handle = RequestHandle;
+		ContactRequest.SingleContactResultDelegate = ResultDelegate;
+		return MLHandleToFGuid(RequestHandle);
+	}
+	else
+	{
+		if (ResultDelegate.IsBound())
+		{
+			ResultDelegate.Broadcast(EMagicLeapContactsOperationStatus::Fail);
+		}
+		Log(EMagicLeapContactsOperationStatus::Fail, FString::Printf(TEXT("MLContactsRequestInsert failed with error '%s'"), UTF8_TO_TCHAR(MLContactsGetResultString(Result))));
+	}
+#endif // WITH_MLSDK
+	return FGuid();
+}
+
+FGuid FMagicLeapContactsPlugin::EditContactAsync(const FMagicLeapContact& InContact, const FSingleContactResultDelegateMulti& ResultDelegate)
+{
+#if WITH_MLSDK
+	if (TryAddPendingTask(EMagicLeapPrivilege::AddressBookWrite, FContactRequest::EType::Edit, InContact, ResultDelegate))
+	{
+		return FGuid();
+	}
+
+	FMagicLeapContact UEContact = InContact;
+	ForceEmailAddressessToLower(UEContact);
+
+	if (!ValidateUEContact(UEContact))
+	{
+		if (ResultDelegate.IsBound())
+		{
+			ResultDelegate.Broadcast(EMagicLeapContactsOperationStatus::Fail);
+		}
+		return FGuid();
+	}
+
+	MLContactsContact MLContact;
+	UEToMLContact(UEContact, MLContact);
+	MLHandle RequestHandle;
+	MLResult Result = MLContactsRequestUpdate(&MLContact, &RequestHandle);
+	DestroyMLContact(MLContact);
+	if (Result == MLResult_Ok)
+	{
+		FContactRequest& ContactRequest = ActiveRequests.AddDefaulted_GetRef();
+		ContactRequest.Handle = RequestHandle;
+		ContactRequest.SingleContactResultDelegate = ResultDelegate;
+		return MLHandleToFGuid(RequestHandle);
+	}
+	else
+	{
+		if (ResultDelegate.IsBound())
+		{
+			ResultDelegate.Broadcast(EMagicLeapContactsOperationStatus::Fail);
+		}
+		Log(EMagicLeapContactsOperationStatus::Fail, FString::Printf(TEXT("MLContactsRequestUpdate failed with error '%s'"), UTF8_TO_TCHAR(MLContactsGetResultString(Result))));
+	}
+#endif // WITH_MLSDK
+	return FGuid();
+}
+
+FGuid FMagicLeapContactsPlugin::DeleteContactAsync(const FMagicLeapContact& InContact, const FSingleContactResultDelegateMulti& ResultDelegate)
+{
+#if WITH_MLSDK
+	if (TryAddPendingTask(EMagicLeapPrivilege::AddressBookWrite, FContactRequest::EType::Delete, InContact, ResultDelegate))
+	{
+		return FGuid();
+	}
+
+	FMagicLeapContact UEContact = InContact;
+	ForceEmailAddressessToLower(UEContact);
+
+	if (!ValidateUEContact(UEContact))
+	{
+		if (ResultDelegate.IsBound())
+		{
+			ResultDelegate.Broadcast(EMagicLeapContactsOperationStatus::Fail);
+		}
+		return FGuid();
+	}
+
+	MLContactsContact MLContact;
+	UEToMLContact(UEContact, MLContact);
+	MLHandle RequestHandle;
+	MLResult Result = MLContactsRequestDelete(&MLContact, &RequestHandle);
+	DestroyMLContact(MLContact);
+	if (Result == MLResult_Ok)
+	{
+		FContactRequest& ContactRequest = ActiveRequests.AddDefaulted_GetRef();
+		ContactRequest.Handle = RequestHandle;
+		ContactRequest.SingleContactResultDelegate = ResultDelegate;
+		return MLHandleToFGuid(RequestHandle);
+	}
+	else
+	{
+		if (ResultDelegate.IsBound())
+		{
+			ResultDelegate.Broadcast(EMagicLeapContactsOperationStatus::Fail);
+		}
+		Log(EMagicLeapContactsOperationStatus::Fail, FString::Printf(TEXT("MLContactsRequestDelete failed with error '%s'"), UTF8_TO_TCHAR(MLContactsGetResultString(Result))));
+	}
+#endif // WITH_MLSDK
+	return FGuid();
+}
+
+FGuid FMagicLeapContactsPlugin::RequestContactsAsync(const FMultipleContactsResultDelegateMulti& ResultDelegate)
+{
+#if WITH_MLSDK
+	if (TryAddPendingTask(EMagicLeapPrivilege::AddressBookRead, FContactRequest::EType::GetAll, ResultDelegate))
+	{
+		return FGuid();
+	}
+
+	MLContactsListArgs Args;
+	MLContactsListArgsInit(&Args);
+	MLHandle RequestHandle;
+	MLResult Result = MLContactsRequestList(&Args, &RequestHandle);
+	if (Result == MLResult_Ok)
+	{
+		FContactRequest& ContactRequest = ActiveRequests.AddDefaulted_GetRef();
+		ContactRequest.Handle = RequestHandle;
+		ContactRequest.MultipleContactsResultDelegate = ResultDelegate;
+		return MLHandleToFGuid(RequestHandle);
+	}
+	else
+	{
+		if (ResultDelegate.IsBound())
+		{
+			TArray<FMagicLeapContact> DummyContacts;
+			ResultDelegate.Broadcast(DummyContacts, EMagicLeapContactsOperationStatus::Fail);
+		}
+		Log(EMagicLeapContactsOperationStatus::Fail, FString::Printf(TEXT("MLContactsRequestList failed with error '%s'"), UTF8_TO_TCHAR(MLContactsGetResultString(Result))));
+	}
+#endif // WITH_MLSDK
+	return FGuid();
+}
+
+FGuid FMagicLeapContactsPlugin::SearchContactsAsync(const FString& Query, EMagicLeapContactsSearchField SearchField, const FMultipleContactsResultDelegateMulti& ResultDelegate)
+{
+#if WITH_MLSDK
+	if (TryAddPendingTask(EMagicLeapPrivilege::AddressBookRead, FContactRequest::EType::Search, ResultDelegate))
+	{
+		return FGuid();
+	}
+
+	MLContactsSearchArgs Args;
+	MLContactsSearchArgsInit(&Args);
+	Args.fields = UESearchFieldToMLSearchField(SearchField);
+	Args.query = TCHAR_TO_UTF8(*Query);
+	MLHandle RequestHandle;
+	MLResult Result = MLContactsRequestSearch(&Args, &RequestHandle);
+	if (Result == MLResult_Ok)
+	{
+		FContactRequest& ContactRequest = ActiveRequests.AddDefaulted_GetRef();
+		ContactRequest.Handle = RequestHandle;
+		ContactRequest.MultipleContactsResultDelegate = ResultDelegate;
+		return MLHandleToFGuid(RequestHandle);
+	}
+	else
+	{
+		if (ResultDelegate.IsBound())
+		{
+			TArray<FMagicLeapContact> DummyContacts;
+			ResultDelegate.Broadcast(DummyContacts, EMagicLeapContactsOperationStatus::Fail);
+		}
+		Log(EMagicLeapContactsOperationStatus::Fail, FString::Printf(TEXT("MLContactsRequestSearch failed with error '%s'"), UTF8_TO_TCHAR(MLContactsGetResultString(Result))));
+	}
+#endif // WITH_MLSDK
+	return FGuid();
+}
+
+bool FMagicLeapContactsPlugin::CancelRequest(const FGuid& RequestHandle)
+{
+	// needs newer api
+// #if WITH_MLSDK
+// 	MLResult Result = MLContactsCancelRequest(FGuidToMLHandle(RequestHandle));
+// 	if (MLContactsResult_Cancelled != Result)
+// 	{
+// 		Log(EMagicLeapContactsOperationStatus::Fail, FString::Printf(TEXT("MLContactsCancelRequest failed with error '%s'"), UTF8_TO_TCHAR(MLContactsGetResultString(Result))));
+// 	}
+// 	return Result == MLContactsResult_Cancelled;
+// #else
+	return false;
+// #endif // WITH_MLSDK
+}
+
+bool FMagicLeapContactsPlugin::SetLogDelegate(const FContactsLogMessageMulti& InLogDelegate)
+{
+	LogDelegate = InLogDelegate;
+	return true;
+}
+
+#if WITH_MLSDK
+EMagicLeapContactsOperationStatus FMagicLeapContactsPlugin::MLOpStatusToUEOpStatus(MLContactsOperationStatus InMLOpStatus)
+{
+	EMagicLeapContactsOperationStatus Status = EMagicLeapContactsOperationStatus::Fail;
+
+	switch (InMLOpStatus)
+	{
+	case MLContactsOperationStatus_Success: Status = EMagicLeapContactsOperationStatus::Success; break;
+	case MLContactsOperationStatus_Fail: Status = EMagicLeapContactsOperationStatus::Fail; break;
+	case MLContactsOperationStatus_Duplicate: Status = EMagicLeapContactsOperationStatus::Duplicate; break;
+	case MLContactsOperationStatus_NotFound: Status = EMagicLeapContactsOperationStatus::NotFound; break;
+	}
+
+	return Status;
+}
+
+MLContactsSearchField FMagicLeapContactsPlugin::UESearchFieldToMLSearchField(EMagicLeapContactsSearchField InSearchField)
+{
+	MLContactsSearchField SearchField = MLContactsSearchField_All;
+
+	switch (InSearchField)
+	{
+	case EMagicLeapContactsSearchField::Name: SearchField = MLContactsSearchField_Name; break;
+	case EMagicLeapContactsSearchField::Phone: SearchField = MLContactsSearchField_Phone; break;
+	case EMagicLeapContactsSearchField::Email: SearchField = MLContactsSearchField_Email; break;
+	case EMagicLeapContactsSearchField::All: SearchField = MLContactsSearchField_All; break;
+	}
+
+	return SearchField;
+}
+
+void FMagicLeapContactsPlugin::UEToMLContact(const FMagicLeapContact& InUEContact, MLContactsContact& OutMLContact)
+{
+	MLContactsContactInit(&OutMLContact);
+	OutMLContact.id = static_cast<const char*>(FMemory::Malloc(InUEContact.Id.Len() + 1));
+	FMemory::Memzero((void*)OutMLContact.id, InUEContact.Id.Len() + 1);
+	FMemory::Memcpy((void*)OutMLContact.id, TCHAR_TO_UTF8(*InUEContact.Id), InUEContact.Id.Len());
+	OutMLContact.name = static_cast<const char*>(FMemory::Malloc(InUEContact.Name.Len() + 1));
+	FMemory::Memzero((void*)OutMLContact.name, InUEContact.Name.Len() + 1);
+	FMemory::Memcpy((void*)OutMLContact.name, TCHAR_TO_UTF8(*InUEContact.Name), InUEContact.Name.Len());
+
+	OutMLContact.phone_number_list.count = static_cast<size_t>(InUEContact.PhoneNumbers.Num());
+	MLContactsTaggedAttribute** MLPhoneNumbers = new MLContactsTaggedAttribute*[InUEContact.PhoneNumbers.Num()];
+	OutMLContact.phone_number_list.items = MLPhoneNumbers;
+	for (int32 iPhoneNumber = 0; iPhoneNumber < InUEContact.PhoneNumbers.Num(); ++iPhoneNumber)
+	{
+		const FMagicLeapTaggedAttribute& UEPhoneNumber = InUEContact.PhoneNumbers[iPhoneNumber];
+		MLContactsTaggedAttribute* MLPhoneNumber = new MLContactsTaggedAttribute;
+		MLPhoneNumber->tag = static_cast<const char*>(FMemory::Malloc(UEPhoneNumber.Tag.Len() + 1));
+		FMemory::Memzero((void*)MLPhoneNumber->tag, UEPhoneNumber.Tag.Len() + 1);
+		FMemory::Memcpy((void*)MLPhoneNumber->tag, TCHAR_TO_UTF8(*UEPhoneNumber.Tag), UEPhoneNumber.Tag.Len());
+		MLPhoneNumber->value = static_cast<const char*>(FMemory::Malloc(UEPhoneNumber.Value.Len() + 1));
+		FMemory::Memzero((void*)MLPhoneNumber->value, UEPhoneNumber.Value.Len() + 1);
+		FMemory::Memcpy((void*)MLPhoneNumber->value, TCHAR_TO_UTF8(*UEPhoneNumber.Value), UEPhoneNumber.Value.Len());
+		MLPhoneNumbers[iPhoneNumber] = MLPhoneNumber;
+	}
+
+	OutMLContact.email_address_list.count = static_cast<size_t>(InUEContact.EmailAddresses.Num());
+	MLContactsTaggedAttribute** MLEmailAddresses = new MLContactsTaggedAttribute*[InUEContact.EmailAddresses.Num()];
+	OutMLContact.email_address_list.items = MLEmailAddresses;
+	for (int32 iEmailAddress = 0; iEmailAddress < InUEContact.EmailAddresses.Num(); ++iEmailAddress)
+	{
+		const FMagicLeapTaggedAttribute& UEEmailAddress = InUEContact.EmailAddresses[iEmailAddress];
+		MLContactsTaggedAttribute* MLEmailAddress = new MLContactsTaggedAttribute;
+		MLEmailAddress->tag = static_cast<const char*>(FMemory::Malloc(UEEmailAddress.Tag.Len() + 1));
+		FMemory::Memzero((void*)MLEmailAddress->tag, UEEmailAddress.Tag.Len() + 1);
+		FMemory::Memcpy((void*)MLEmailAddress->tag, TCHAR_TO_UTF8(*UEEmailAddress.Tag), UEEmailAddress.Tag.Len());
+		MLEmailAddress->value = static_cast<const char*>(FMemory::Malloc(UEEmailAddress.Value.Len() + 1));
+		FMemory::Memzero((void*)MLEmailAddress->value, UEEmailAddress.Value.Len() + 1);
+		FMemory::Memcpy((void*)MLEmailAddress->value, TCHAR_TO_UTF8(*UEEmailAddress.Value), UEEmailAddress.Value.Len());
+		MLEmailAddresses[iEmailAddress] = MLEmailAddress;
+	}
+}
+
+void FMagicLeapContactsPlugin::DestroyMLContact(MLContactsContact& OutMLContact)
+{
+	FMemory::Free((void*)OutMLContact.id);
+	FMemory::Free((void*)OutMLContact.name);
+
+	for (size_t iPhoneNumber = 0; iPhoneNumber < OutMLContact.phone_number_list.count; ++iPhoneNumber)
+	{
+		MLContactsTaggedAttribute* MLPhoneNumber = OutMLContact.phone_number_list.items[iPhoneNumber];
+		FMemory::Free((void*)MLPhoneNumber->tag);
+		FMemory::Free((void*)MLPhoneNumber->value);
+		delete MLPhoneNumber;
+	}
+	delete OutMLContact.phone_number_list.items;
+
+	for (size_t iEmailAddress = 0; iEmailAddress < OutMLContact.email_address_list.count; ++iEmailAddress)
+	{
+		MLContactsTaggedAttribute* MLEmailAddress = OutMLContact.email_address_list.items[iEmailAddress];
+		FMemory::Free((void*)MLEmailAddress->tag);
+		FMemory::Free((void*)MLEmailAddress->value);
+		delete MLEmailAddress;
+	}
+	delete OutMLContact.email_address_list.items;
+}
+
+void FMagicLeapContactsPlugin::MLToUEContact(const MLContactsContact* InMLContact, FMagicLeapContact& OutUEContact)
+{
+	OutUEContact.Id = UTF8_TO_TCHAR(InMLContact->id);
+	OutUEContact.Name = UTF8_TO_TCHAR(InMLContact->name);
+
+	const int32 NumPhoneNumbers = static_cast<int32>(InMLContact->phone_number_list.count);
+	OutUEContact.PhoneNumbers.AddZeroed(NumPhoneNumbers);
+	for (int32 iPhoneNumber = 0; iPhoneNumber < NumPhoneNumbers; ++iPhoneNumber)
+	{
+		FMagicLeapTaggedAttribute& UEPhoneNumber = OutUEContact.PhoneNumbers[iPhoneNumber];
+		const MLContactsTaggedAttribute* MLPhoneNumber = InMLContact->phone_number_list.items[iPhoneNumber];
+		UEPhoneNumber.Tag = UTF8_TO_TCHAR(MLPhoneNumber->tag);
+		UEPhoneNumber.Value = UTF8_TO_TCHAR(MLPhoneNumber->value);
+	}
+
+	const int32 NumEmailAddresses = static_cast<int32>(InMLContact->email_address_list.count);
+	OutUEContact.EmailAddresses.AddZeroed(NumEmailAddresses);
+	for (int32 iEmailAddress = 0; iEmailAddress < NumEmailAddresses; ++iEmailAddress)
+	{
+		FMagicLeapTaggedAttribute& UEAddress = OutUEContact.EmailAddresses[iEmailAddress];
+		const MLContactsTaggedAttribute* MLEmailAddress = InMLContact->email_address_list.items[iEmailAddress];
+		UEAddress.Tag = UTF8_TO_TCHAR(MLEmailAddress->tag);
+		UEAddress.Value = UTF8_TO_TCHAR(MLEmailAddress->value);
+	}
+}
+
+MLHandle FMagicLeapContactsPlugin::FGuidToMLHandle(const FGuid& InGuid)
+{
+	return static_cast<uint64_t>(InGuid.B) << 32 | InGuid.A;
+}
+
+FGuid FMagicLeapContactsPlugin::MLHandleToFGuid(const MLHandle& InHandle)
+{
+	return FGuid(InHandle, InHandle >> 32, 0, 0);
+}
+#endif // WITH_MLSDK
+
+bool FMagicLeapContactsPlugin::ValidateUEContact(const FMagicLeapContact& InUEContact)
+{
+	bool bIsValid = true;
+
+	if (InUEContact.Name.IsEmpty())
+	{
+		Log(EMagicLeapContactsOperationStatus::Fail, TEXT("Contact name is required!"));
+		bIsValid = false;
+	}
+	else if (InUEContact.Name.Len() > MAX_NAME_LENGTH)
+	{
+		Log(EMagicLeapContactsOperationStatus::Fail, FString::Printf(TEXT("Contact name '%s' is too long. Max length is: %d"), *InUEContact.Name, MAX_NAME_LENGTH));
+		bIsValid = false;
+	}
+	if (0 == InUEContact.PhoneNumbers.Num() && 0 == InUEContact.EmailAddresses.Num())
+	{
+		Log(EMagicLeapContactsOperationStatus::Fail, TEXT("Either a phone number or e-mail address is required"));
+		bIsValid = false;
+	}
+	if (InUEContact.PhoneNumbers.Num() > MAX_PHONE_COUNT)
+	{
+		Log(EMagicLeapContactsOperationStatus::Fail, FString::Printf(TEXT("Too many phone numbers have been added. Max count is: %d"), MAX_PHONE_COUNT));
+		bIsValid = false;
+	}
+	else
+	{
+		for (const FMagicLeapTaggedAttribute& Phone : InUEContact.PhoneNumbers)
+		{
+			if (!Phone.Tag.IsEmpty() && Phone.Tag.Len() > MAX_TAG_LENGTH)
+			{
+				Log(EMagicLeapContactsOperationStatus::Fail, FString::Printf(TEXT("Phone Tag '%s' is too long. Limit is: %d"), *Phone.Tag, MAX_TAG_LENGTH));
+				bIsValid = false;
+			}
+
+			FRegexMatcher PhonePatternMatcher(PhoneNumberPattern, Phone.Value);
+			if (Phone.Value.IsEmpty() || !PhonePatternMatcher.FindNext())
+			{
+				Log(EMagicLeapContactsOperationStatus::Fail, FString::Printf(TEXT("Phone number '%s' is invalid."), *Phone.Value));
+				bIsValid = false;
+			}
+			else if (Phone.Value.Len() > MAX_PHONE_LENGTH)
+			{
+				Log(EMagicLeapContactsOperationStatus::Fail, FString::Printf(TEXT("Phone number '%s' is too long. Max length is: %d"), *Phone.Value, MAX_PHONE_LENGTH));
+				bIsValid = false;
+			}
+		}
+	}
+	if (InUEContact.EmailAddresses.Num() > MAX_EMAIL_COUNT)
+	{
+		Log(EMagicLeapContactsOperationStatus::Fail, FString::Printf(TEXT("Too many email addresses have been added. Limit is: %d"), MAX_EMAIL_COUNT));
+		bIsValid = false;
+	}
+	else
+	{
+		for (const FMagicLeapTaggedAttribute& Email : InUEContact.EmailAddresses)
+		{
+			if (!Email.Tag.IsEmpty() && Email.Tag.Len() > MAX_TAG_LENGTH)
+			{
+				Log(EMagicLeapContactsOperationStatus::Fail, FString::Printf(TEXT("Email Tag is too long. Max length is: %d"), MAX_TAG_LENGTH));
+				bIsValid = false;
+			}
+			FRegexMatcher EmailPatternMatcher(EmailPattern, Email.Value);
+			if (Email.Value.IsEmpty() || !EmailPatternMatcher.FindNext())
+			{
+				Log(EMagicLeapContactsOperationStatus::Fail, FString::Printf(TEXT("Email address '%s' is invalid."), *Email.Value));
+				bIsValid = false;
+			}
+			else if (Email.Value.Len() > MAX_EMAIL_LENGTH)
+			{
+				Log(EMagicLeapContactsOperationStatus::Fail, FString::Printf(TEXT("Email address '%s' is too long.  Max length is %d"), *Email.Value, MAX_EMAIL_LENGTH));
+				bIsValid = false;
+			}
+		}
+	}
+
+	return bIsValid;
+}
+
+void FMagicLeapContactsPlugin::ForceEmailAddressessToLower(FMagicLeapContact& InUEContact)
+{
+	for (FMagicLeapTaggedAttribute& Email : InUEContact.EmailAddresses)
+	{
+		if (!Email.Value.IsEmpty())
+		{
+			Email.Value.ToLowerInline();
+		}
+	}
+}
+
+bool FMagicLeapContactsPlugin::TryAddPendingTask(EMagicLeapPrivilege InRequiredPrivilege, FContactRequest::EType InRequestType, const FMagicLeapContact& InContact, const FSingleContactResultDelegateMulti& InResultDelegate)
+{
+	bool bAddedTask = false;
+	EPrivilegeState PrivilegeStatus = PrivilegesManager.GetPrivilegeStatus(InRequiredPrivilege, false);
+	if (PrivilegeStatus != EPrivilegeState::Granted)
+	{
+		if (PrivilegeStatus == EPrivilegeState::Pending)
+		{
+			FContactRequest& ContactRequest = PendingRequests.AddDefaulted_GetRef();
+			ContactRequest.Type = InRequestType;
+			ContactRequest.Contact = InContact;
+			ContactRequest.RequiredPrivilege = InRequiredPrivilege;
+			ContactRequest.SingleContactResultDelegate = InResultDelegate;
+		}
+		else if (InResultDelegate.IsBound())
+		{
+			InResultDelegate.Broadcast(EMagicLeapContactsOperationStatus::Fail);
+			FString ErrorMsg = FString::Printf(TEXT("TryAddPendingTask failed due to privilege '%s' having status '%s'"), *PrivilegesManager.PrivilegeToString(InRequiredPrivilege), PrivilegesManager.PrivilegeStateToString(PrivilegeStatus));
+			Log(EMagicLeapContactsOperationStatus::Fail, ErrorMsg);
+		}
+
+		bAddedTask = true;
+	}
+
+	return bAddedTask;
+}
+
+bool FMagicLeapContactsPlugin::TryAddPendingTask(EMagicLeapPrivilege InRequiredPrivilege, FContactRequest::EType InRequestType, const FMultipleContactsResultDelegateMulti& InResultDelegate)
+{
+	bool bAddedTask = false;
+	EPrivilegeState PrivilegeStatus = PrivilegesManager.GetPrivilegeStatus(InRequiredPrivilege, false);
+	if (PrivilegeStatus != EPrivilegeState::Granted)
+	{
+		if (PrivilegeStatus == EPrivilegeState::Pending)
+		{
+			FContactRequest& ContactRequest = PendingRequests.AddDefaulted_GetRef();
+			ContactRequest.Type = InRequestType;
+			ContactRequest.RequiredPrivilege = InRequiredPrivilege;
+			ContactRequest.MultipleContactsResultDelegate = InResultDelegate;
+			bAddedTask = true;
+		}
+		else if (InResultDelegate.IsBound())
+		{
+			TArray<FMagicLeapContact> DummyContacts;
+			InResultDelegate.Broadcast(DummyContacts, EMagicLeapContactsOperationStatus::Fail);
+			FString ErrorMsg = FString::Printf(TEXT("TryAddPendingTask failed due to privilege '%s' having status '%s'"), *PrivilegesManager.PrivilegeToString(InRequiredPrivilege), PrivilegesManager.PrivilegeStateToString(PrivilegeStatus));
+			Log(EMagicLeapContactsOperationStatus::Fail, ErrorMsg);
+		}
+	}
+
+	return bAddedTask;
+}
+
+void FMagicLeapContactsPlugin::ProcessPendingRequests()
+{
+	for (int32 iPendingRequest = PendingRequests.Num() - 1; iPendingRequest > -1; --iPendingRequest)
+	{
+		const FContactRequest& PendingRequest = PendingRequests[iPendingRequest];
+		EPrivilegeState PrivilegeStatus = PrivilegesManager.GetPrivilegeStatus(PendingRequest.RequiredPrivilege, false);
+
+		switch (PrivilegeStatus)
+		{
+		case EPrivilegeState::NotYetRequested:
+		{
+			checkf(false, TEXT("Invalid request status 'EPrivilegeState::NotYetRequested' encountered in pending request!"));
+		}
+		break;
+
+		case EPrivilegeState::Pending:
+		{
+			continue;
+		}
+		break;
+
+		case EPrivilegeState::Granted:
+		{
+			switch (PendingRequest.Type)
+			{
+			case FContactRequest::EType::Add:
+			{
+				AddContactAsync(PendingRequest.Contact, PendingRequest.SingleContactResultDelegate);
+			}
+			break;
+			case FContactRequest::EType::Edit:
+			{
+				EditContactAsync(PendingRequest.Contact, PendingRequest.SingleContactResultDelegate);
+			}
+			break;
+			case FContactRequest::EType::Delete:
+			{
+				DeleteContactAsync(PendingRequest.Contact, PendingRequest.SingleContactResultDelegate);
+			}
+			break;
+			case FContactRequest::EType::GetAll:
+			{
+				RequestContactsAsync(PendingRequest.MultipleContactsResultDelegate);
+			}
+			break;
+			case FContactRequest::EType::Search:
+			{
+				checkf(false, TEXT("Unexpected pending request 'Search' encountered!"));
+			}
+			break;
+			case FContactRequest::EType::Cancel:
+			{
+				checkf(false, TEXT("Unexpected pending request 'Cancel' encountered!"));
+			}
+			break;
+			}
+
+			PendingRequests.RemoveAt(iPendingRequest);
+		}
+		break;
+
+		case EPrivilegeState::Denied:
+		{
+			PendingRequests.RemoveAt(iPendingRequest);
+		}
+		break;
+
+		case EPrivilegeState::Error:
+		{
+			PendingRequests.RemoveAt(iPendingRequest);
+		}
+		break;
+		}
+	}
+}
+
+void FMagicLeapContactsPlugin::Log(EMagicLeapContactsOperationStatus OpStatus, const FString& LogString)
+{
+	if (OpStatus != EMagicLeapContactsOperationStatus::Success)
+	{
+		UE_LOG(LogMagicLeapContacts, Error, TEXT("%s"), *LogString);
+	}
+	else
+	{
+		UE_LOG(LogMagicLeapContacts, Log, TEXT("%s"), *LogString);
+	}
+
+	if (LogDelegate.IsBound())
+	{
+		LogDelegate.Broadcast(LogString, OpStatus);
+	}
+}
+
+IMPLEMENT_MODULE(FMagicLeapContactsPlugin, MagicLeapContacts);
diff --git a/Engine/Plugins/Lumin/MagicLeapContacts/Source/Private/MagicLeapContactsPlugin.h b/Engine/Plugins/Lumin/MagicLeapContacts/Source/Private/MagicLeapContactsPlugin.h
new file mode 100644
--- /dev/null
+++ b/Engine/Plugins/Lumin/MagicLeapContacts/Source/Private/MagicLeapContactsPlugin.h
@@ -0,0 +1,90 @@
+// Copyright 1998-2019 Epic Games, Inc. All Rights Reserved.
+
+#pragma once
+
+#include "CoreMinimal.h"
+#include "Internationalization/Regex.h"
+#include "IMagicLeapContactsPlugin.h"
+#include "MagicLeapContactsTypes.h"
+#include "AppEventHandler.h"
+#include "Lumin/CAPIShims/LuminAPIContacts.h"
+#include "MagicLeapPluginUtil.h"
+
+DECLARE_LOG_CATEGORY_EXTERN(LogMagicLeapContacts, Verbose, All);
+
+class FMagicLeapContactsPlugin : public IMagicLeapContactsPlugin
+{
+public:
+	FMagicLeapContactsPlugin();
+	void StartupModule() override;
+	void ShutdownModule() override;
+	bool Tick(float DeltaTime) override;
+
+	bool Startup();
+	bool Shutdown();
+	FGuid AddContactAsync(const FMagicLeapContact& Contact, const FSingleContactResultDelegateMulti& ResultDelegate);
+	FGuid EditContactAsync(const FMagicLeapContact& Contact, const FSingleContactResultDelegateMulti& ResultDelegate);
+	FGuid DeleteContactAsync(const FMagicLeapContact& Contact, const FSingleContactResultDelegateMulti& ResultDelegate);
+	FGuid RequestContactsAsync(const FMultipleContactsResultDelegateMulti& ResultDelegate);
+	FGuid SearchContactsAsync(const FString& Query, EMagicLeapContactsSearchField SearchField, const FMultipleContactsResultDelegateMulti& ResultDelegate);
+	bool CancelRequest(const FGuid& RequestHandle);
+	bool SetLogDelegate(const FContactsLogMessageMulti& LogDelegate);
+
+private:
+	// JMC TODO: Validation needs to come from the api.
+	const FRegexPattern PhoneNumberPattern;
+	const FRegexPattern EmailPattern;
+	const int32 MAX_NAME_LENGTH;
+	const int32 MAX_PHONE_LENGTH;
+	const int32 MAX_EMAIL_LENGTH;
+	const int32 MAX_TAG_LENGTH;
+	const int32 MAX_PHONE_COUNT;
+	const int32 MAX_EMAIL_COUNT;
+
+	MagicLeap::IAppEventHandler PrivilegesManager;
+	FMagicLeapAPISetup APISetup;
+	FTickerDelegate TickDelegate;
+	FDelegateHandle TickDelegateHandle;
+	struct FContactRequest
+	{
+		enum class EType
+		{
+			Add,
+			Edit,
+			Delete,
+			GetAll,
+			Search,
+			Cancel
+		};
+
+		EType Type;
+		FMagicLeapContact Contact;
+		EMagicLeapPrivilege RequiredPrivilege;
+		FSingleContactResultDelegateMulti SingleContactResultDelegate;
+		FMultipleContactsResultDelegateMulti MultipleContactsResultDelegate;
+#if WITH_MLSDK
+		MLHandle Handle;
+#endif // WITH_MLSDK
+	};
+	TArray<FContactRequest> PendingRequests;
+	TArray<FContactRequest> ActiveRequests;
+	FContactsLogMessageMulti LogDelegate;
+
+#if WITH_MLSDK
+	EMagicLeapContactsOperationStatus MLOpStatusToUEOpStatus(MLContactsOperationStatus InMLOpStatus);
+	MLContactsSearchField UESearchFieldToMLSearchField(EMagicLeapContactsSearchField InSearchField);
+	void UEToMLContact(const FMagicLeapContact& InUEContact, MLContactsContact& OutMLContact);
+	void MLToUEContact(const MLContactsContact* InMLContact, FMagicLeapContact& OutUEContact);
+	void DestroyMLContact(MLContactsContact& OutMLContact);
+	MLHandle FGuidToMLHandle(const FGuid& InGuid);
+	FGuid MLHandleToFGuid(const MLHandle& InHandle);
+#endif // WITH_MLSDK
+	bool ValidateUEContact(const FMagicLeapContact& InUEContact);
+	void ForceEmailAddressessToLower(FMagicLeapContact& InUEContact);
+	bool TryAddPendingTask(EMagicLeapPrivilege InRequiredPrivilege, FContactRequest::EType InRequestType, const FMagicLeapContact& InContact, const FSingleContactResultDelegateMulti& InResultDelegate);
+	bool TryAddPendingTask(EMagicLeapPrivilege InRequiredPrivilege, FContactRequest::EType InRequestType, const FMultipleContactsResultDelegateMulti& InResultDelegate);
+	void ProcessPendingRequests();
+	void Log(EMagicLeapContactsOperationStatus OpStatus, const FString& LogString);
+};
+
+#define GET_MAGIC_LEAP_CONTACTS_PLUGIN() static_cast<FMagicLeapContactsPlugin*>(&IMagicLeapContactsPlugin::Get())
diff --git a/Engine/Plugins/Lumin/MagicLeapContacts/Source/Public/IMagicLeapContactsPlugin.h b/Engine/Plugins/Lumin/MagicLeapContacts/Source/Public/IMagicLeapContactsPlugin.h
new file mode 100644
--- /dev/null
+++ b/Engine/Plugins/Lumin/MagicLeapContacts/Source/Public/IMagicLeapContactsPlugin.h
@@ -0,0 +1,38 @@
+// Copyright 1998-2019 Epic Games, Inc. All Rights Reserved.
+
+#pragma once
+
+#include "Modules/ModuleManager.h"
+#include "Containers/Ticker.h"
+
+/**
+ * The public interface to this module.  In most cases, this interface is only public to sibling modules
+ * within this plugin.
+ */
+class IMagicLeapContactsPlugin : public IModuleInterface
+{
+public:
+
+	/**
+	 * Singleton-like access to this module's interface.  This is just for convenience!
+	 * Beware of calling this during the shutdown phase, though.  Your module might have been unloaded already.
+	 *
+	 * @return Returns singleton instance, loading the module on demand if needed
+	 */
+	static inline IMagicLeapContactsPlugin& Get()
+	{
+		return FModuleManager::LoadModuleChecked<IMagicLeapContactsPlugin>("MagicLeapContacts");
+	}
+
+	/**
+	 * Checks to see if this module is loaded and ready.  It is only valid to call Get() if IsAvailable() returns true.
+	 *
+	 * @return True if the module is loaded and ready to use
+	 */
+	static inline bool IsAvailable()
+	{
+		return FModuleManager::Get().IsModuleLoaded("MagicLeapContacts");
+	}
+
+	virtual bool Tick(float DeltaTime) = 0;
+};
diff --git a/Engine/Plugins/Lumin/MagicLeapContacts/Source/Public/MagicLeapContactsComponent.h b/Engine/Plugins/Lumin/MagicLeapContacts/Source/Public/MagicLeapContactsComponent.h
new file mode 100644
--- /dev/null
+++ b/Engine/Plugins/Lumin/MagicLeapContacts/Source/Public/MagicLeapContactsComponent.h
@@ -0,0 +1,92 @@
+// Copyright 1998-2019 Epic Games, Inc. All Rights Reserved.
+
+#pragma once
+
+#include "Engine/Engine.h"
+#include "Components/ActorComponent.h"
+#include "MagicLeapContactsTypes.h"
+#include "MagicLeapContactsComponent.generated.h"
+
+/**
+	Component that provides access to the Contacts API functionality.
+*/
+UCLASS(ClassGroup = MagicLeap, BlueprintType, Blueprintable, EditInlineNew, meta = (BlueprintSpawnableComponent))
+class MAGICLEAPCONTACTS_API UMagicLeapContactsComponent : public UActorComponent
+{
+	GENERATED_BODY()
+
+public:
+	/** Starts up the Contacts API. */
+	void BeginPlay() override;
+
+	/** Shuts down the Contacts API */
+	void EndPlay(const EEndPlayReason::Type EndPlayReason) override;
+
+	/**
+		Initiates the creation of a new contact.
+		@param Contact The contact to be created.
+		@return A unique identifier for this request (required if request needs to be cancelled).
+	*/
+	UFUNCTION(BlueprintCallable, Category = "Contacts | MagicLeap")
+	FGuid AddContactAsync(const FMagicLeapContact& Contact);
+
+	/**
+		Initiates the update of an existing contact.
+		@param Contact The contact to be updated.
+		@return A unique identifier for this request (required if request needs to be cancelled).
+	*/
+	UFUNCTION(BlueprintCallable, Category = "Contacts | MagicLeap")
+	FGuid EditContactAsync(const FMagicLeapContact& Contact);
+
+	/**
+		Initiates the deletion of an existing contact.
+		@param Contact The contact to be deleted.
+		@return A unique identifier for this request (required if request needs to be cancelled).
+	*/
+	UFUNCTION(BlueprintCallable, Category = "Contacts | MagicLeap")
+	FGuid DeleteContactAsync(const FMagicLeapContact& Contact);
+
+	/**
+		Initiates the retrieval of the entire contacts list from the cloud.
+		@return A unique identifier for this request (required if request needs to be cancelled).
+	*/
+	UFUNCTION(BlueprintCallable, Category = "Contacts | MagicLeap")
+	FGuid RequestContactsAsync();
+
+	/**
+		Initiates a search for contacts with a given query across specified fields.
+		@param Query The search string to look for instances of.
+		@param SearchField The field within the contact to match the query against.
+		@return A unique identifier for this request (required if request needs to be cancelled).
+	*/
+	UFUNCTION(BlueprintCallable, Category = "Contacts | MagicLeap")
+	FGuid SearchContactsAsync(const FString& Query, EMagicLeapContactsSearchField SearchField);
+
+	/**
+		Cancels a request corresponding to the FGuid.
+		@param RequestHandle The unique identifier of the request (returned by all contact request functions).
+		@return True if the cancellation succeeded, false otherwise.
+	*/
+	UFUNCTION(BlueprintCallable, Category = "Contacts | MagicLeap")
+	bool CancelRequest(const FGuid& RequestHandle);
+
+private:
+	// Delegate instances
+	UPROPERTY(BlueprintAssignable, Category = "Contacts | MagicLeap", meta = (AllowPrivateAccess = true))
+	FSingleContactResultDelegateMulti OnAddContactResult;
+
+	UPROPERTY(BlueprintAssignable, Category = "Contacts | MagicLeap", meta = (AllowPrivateAccess = true))
+	FSingleContactResultDelegateMulti OnEditContactResult;
+
+	UPROPERTY(BlueprintAssignable, Category = "Contacts | MagicLeap", meta = (AllowPrivateAccess = true))
+	FSingleContactResultDelegateMulti OnDeleteContactResult;
+
+	UPROPERTY(BlueprintAssignable, Category = "Contacts | MagicLeap", meta = (AllowPrivateAccess = true))
+	FMultipleContactsResultDelegateMulti OnRequestContactsResult;
+
+	UPROPERTY(BlueprintAssignable, Category = "Contacts | MagicLeap", meta = (AllowPrivateAccess = true))
+	FMultipleContactsResultDelegateMulti OnSearchContactsResult;
+
+	UPROPERTY(BlueprintAssignable, Category = "Contacts | MagicLeap", meta = (AllowPrivateAccess = true))
+	FContactsLogMessageMulti OnLogMessage;
+};
diff --git a/Engine/Plugins/Lumin/MagicLeapContacts/Source/Public/MagicLeapContactsFunctionLibrary.h b/Engine/Plugins/Lumin/MagicLeapContacts/Source/Public/MagicLeapContactsFunctionLibrary.h
new file mode 100644
--- /dev/null
+++ b/Engine/Plugins/Lumin/MagicLeapContacts/Source/Public/MagicLeapContactsFunctionLibrary.h
@@ -0,0 +1,84 @@
+// Copyright 1998-2019 Epic Games, Inc. All Rights Reserved.
+#pragma once
+
+#include "Engine/Engine.h"
+#include "Kismet/BlueprintFunctionLibrary.h"
+#include "MagicLeapContactsTypes.h"
+#include "MagicLeapContactsFunctionLibrary.generated.h"
+
+UCLASS(ClassGroup = MagicLeap)
+class MAGICLEAPCONTACTS_API UMagicLeapContactsFunctionLibrary : public UBlueprintFunctionLibrary
+{
+	GENERATED_BODY()
+
+public:
+
+	/** Initialize all necessary resources for using the Contacts API. */
+	UFUNCTION(BlueprintCallable, Category = "Contacts Function Library | MagicLeap")
+	static bool Startup();
+
+	/** Deinitialize all resources for this API. */
+	UFUNCTION(BlueprintCallable, Category = "Contacts Function Library | MagicLeap")
+	static bool Shutdown();
+
+	/**
+		Initiates the creation of a new contact.
+		@param Contact The contact to be created.
+		@param ResultDelegate The delegate to be notified upon creation of the contact.
+		@return A unique identifier for this request (required if request needs to be cancelled).
+	*/
+	UFUNCTION(BlueprintCallable, Category = "Contacts Function Library | MagicLeap")
+	static FGuid AddContactAsync(const FMagicLeapContact& Contact, const FSingleContactResultDelegate& ResultDelegate);
+
+	/**
+		Initiates the update of an existing contact.
+		@param Contact The contact to be updated.
+		@param ResultDelegate The delegate to be notified upon update of the contact.
+		@return A unique identifier for this request (required if request needs to be cancelled).
+	*/
+	UFUNCTION(BlueprintCallable, Category = "Contacts Function Library | MagicLeap")
+	static FGuid EditContactAsync(const FMagicLeapContact& Contact, const FSingleContactResultDelegate& ResultDelegate);
+
+	/** 
+		Initiates the deletion of an existing contact.
+		@param Contact The contact to be deleted.
+		@param ResultDelegate The delegate to be notified upon deletion of the contact.
+		@return A unique identifier for this request (required if request needs to be cancelled).
+	*/
+	UFUNCTION(BlueprintCallable, Category = "Contacts Function Library | MagicLeap")
+	static FGuid DeleteContactAsync(const FMagicLeapContact& Contact, const FSingleContactResultDelegate& ResultDelegate);
+
+	/** 
+		Initiates the retrieval of the entire contacts list from the cloud.
+		@param ResultDelegate The delegate to be notified once the contacts list has been retrieved from the cloud.
+		@return A unique identifier for this request (required if request needs to be cancelled).
+	*/
+	UFUNCTION(BlueprintCallable, Category = "Contacts Function Library | MagicLeap")
+	static FGuid RequestContactsAsync(const FMultipleContactsResultDelegate& ResultDelegate);
+
+	/**
+		Initiates a search for contacts with a given query across specified fields.
+		@param Query The search string to look for instances of.
+		@param SearchField The field within the contact to match the query against.
+		@param ResultDelegate The delegate to be notified upon completion of the query.
+		@return A unique identifier for this request (required if request needs to be cancelled).
+	*/
+	UFUNCTION(BlueprintCallable, Category = "Contacts Function Library | MagicLeap")
+	static FGuid SearchContactsAsync(const FString& Query, EMagicLeapContactsSearchField SearchField, const FMultipleContactsResultDelegate& ResultDelegate);
+
+	/**
+		Cancels a request corresponding to the FGuid.
+		@param RequestHandle The unique identifier of the request (returned by all contact request functions).
+		@return True if the cancellation succeeded, false otherwise.
+	*/
+	UFUNCTION(BlueprintCallable, Category = "Contacts Function Library | MagicLeap")
+	static bool CancelRequest(const FGuid& RequestHandle);
+
+	/**
+		Sets the delegate by which the system can pass log messages back to the calling blueprint.
+		@param LogDelegate The delegate by which the system will return log messages to the calling blueprint.
+		@return True if the call succeeds, false otherwise.
+	*/
+	UFUNCTION(BlueprintCallable, Category = "Contacts Function Library | MagicLeap")
+	static bool SetLogDelegate(const FContactsLogMessage& LogDelegate);
+};
diff --git a/Engine/Plugins/Lumin/MagicLeapContacts/Source/Public/MagicLeapContactsTypes.h b/Engine/Plugins/Lumin/MagicLeapContacts/Source/Public/MagicLeapContactsTypes.h
new file mode 100644
--- /dev/null
+++ b/Engine/Plugins/Lumin/MagicLeapContacts/Source/Public/MagicLeapContactsTypes.h
@@ -0,0 +1,108 @@
+// Copyright 1998-2019 Epic Games, Inc. All Rights Reserved.
+#pragma once
+
+#include "Engine/Engine.h"
+#include "MagicLeapContactsTypes.generated.h"
+
+/** Result types for contacts requests. */
+UENUM(BlueprintType)
+enum class EMagicLeapContactsResult : uint8
+{
+	/** This handle is not yet recognized. */
+	HandleNotFound,
+	/** Request is completed, its corresponding result has been returned, and its related resources are marked for deletion. */
+	Completed,
+	/** Request is successfully cancelled. */
+	Cancelled,
+	/** Request failed due to system being in an illegal state, for e.g., when the user hasn't successfully logged-in. */
+	IllegalState,
+};
+
+/** Search query field types. */
+UENUM(BlueprintType)
+enum class EMagicLeapContactsSearchField : uint8
+{
+	/** Search field for nickname. */
+	Name,
+	/** Search field for phone. */
+	Phone,
+	/** Search field for email. */
+	Email,
+	/** Search across all fields. */
+	All,
+};
+
+/** Result types for contacts operations. */
+UENUM(BlueprintType)
+enum class EMagicLeapContactsOperationStatus : uint8
+{
+	/** Operation succeeded. */
+	Success,
+	/** Operation failed. */
+	Fail,
+	/** MagicLeapContact with the details specified for an insert already exists. */
+	Duplicate,
+	/** MagicLeapContact to be deleted/updated doesn't exist. */
+	NotFound,
+};
+
+/**
+	Stores a tagged value, such as phone number or email address.  Optional tag indicates what type
+	of value is stored, e.g. "home", "work", etc.
+*/
+USTRUCT(BlueprintType)
+struct FMagicLeapTaggedAttribute
+{
+	GENERATED_USTRUCT_BODY()
+
+	/** Name of the Tag. */
+	UPROPERTY(EditAnywhere, BlueprintReadWrite, Category = "Contacts|MagicLeap")
+	FString Tag;
+
+	/** Value of this attribute. */
+	UPROPERTY(EditAnywhere, BlueprintReadWrite, Category = "Contacts|MagicLeap")
+	FString Value;
+};
+
+/** Representation of available information for a single contact in an address book. */
+USTRUCT(BlueprintType)
+struct FMagicLeapContact
+{
+	GENERATED_USTRUCT_BODY()
+
+	/**
+		Locally-unique contact identifier.  Generated by the system.  May change across reboots.
+		@note Please do not edit this value.  It is exposed as read/write merely so that it can copied via the make/break functionality.
+	*/
+	UPROPERTY(EditAnywhere, BlueprintReadWrite, Category = "Contacts|MagicLeap")
+	FString Id;
+
+	/** Contacts's name. */
+	UPROPERTY(EditAnywhere, BlueprintReadWrite, Category = "Contacts|MagicLeap")
+	FString Name;
+
+	/** Contacts's phone numbers. */
+	UPROPERTY(EditAnywhere, BlueprintReadWrite, Category = "Contacts|MagicLeap")
+	TArray<FMagicLeapTaggedAttribute> PhoneNumbers;
+
+	/** Contacts's email addresses. */
+	UPROPERTY(EditAnywhere, BlueprintReadWrite, Category = "Contacts|MagicLeap")
+	TArray<FMagicLeapTaggedAttribute> EmailAddresses;
+};
+
+/** Delegate used to convey the result of a single contact operation. */
+DECLARE_DYNAMIC_DELEGATE_OneParam(FSingleContactResultDelegate, EMagicLeapContactsOperationStatus, OpStatus);
+DECLARE_DYNAMIC_MULTICAST_DELEGATE_OneParam(FSingleContactResultDelegateMulti, EMagicLeapContactsOperationStatus, OpStatus);
+
+/** Delegate used to convey the result of a multiple contacts operation. */
+DECLARE_DYNAMIC_DELEGATE_TwoParams(FMultipleContactsResultDelegate, const TArray<FMagicLeapContact>&, Contacts, EMagicLeapContactsOperationStatus, OpStatus);
+DECLARE_DYNAMIC_MULTICAST_DELEGATE_TwoParams(FMultipleContactsResultDelegateMulti, const TArray<FMagicLeapContact>&, Contacts, EMagicLeapContactsOperationStatus, OpStatus);
+
+/**
+	Delegate used to pass log messages from the contacts plugin to the initiating blueprint.
+	@note This is useful if the user wishes to have log messages in 3D space.
+	@param LogMessage A string containing the log message.
+	@param OpStatus The status of the operation associated with the log message.
+*/
+DECLARE_DYNAMIC_DELEGATE_TwoParams(FContactsLogMessage, const FString&, LogMessage, EMagicLeapContactsOperationStatus, OpStatus);
+DECLARE_DYNAMIC_MULTICAST_DELEGATE_TwoParams(FContactsLogMessageMulti, const FString&, LogMessage, EMagicLeapContactsOperationStatus, OpStatus);
diff --git a/Engine/Plugins/Lumin/MagicLeapDispatch/MagicLeapDispatch.uplugin b/Engine/Plugins/Lumin/MagicLeapDispatch/MagicLeapDispatch.uplugin
new file mode 100644
--- /dev/null
+++ b/Engine/Plugins/Lumin/MagicLeapDispatch/MagicLeapDispatch.uplugin
@@ -0,0 +1,42 @@
+{
+	"FileVersion" : 3,
+	"Version" : 1,
+	"VersionName": "1.0",
+	"FriendlyName" : "Magic Leap Dispatch",
+	"CreatedBy" : "Magic Leap, Inc.",
+	"CreatedByURL" : "http://magicleap.com",
+	"Description" : "Support for the Magic Leap Dispatch API.",
+	"Category" : "Lumin",
+	"EnabledByDefault" : false,
+	"DocsURL" : "",
+	"MarketplaceURL" : "",
+	"SupportURL" : "",
+	"CanContainContent" : false,
+	"IsBetaVersion" : false,
+	"Installed" : false,
+	"SupportedTargetPlatforms": [
+		"Lumin",
+		"Win64",
+		"Mac"
+	],
+	"Modules" :
+	[
+		{
+			"Name" : "MagicLeapDispatch",
+			"Type" : "Runtime",
+			"WhitelistPlatforms" :
+			[
+				"Win64",
+				"Mac",
+				"Lumin"
+			]
+		}
+	],
+	"Plugins":
+	[
+		{
+			"Name": "MLSDK",
+			"Enabled": true
+		}
+	]
+}
diff --git a/Engine/Plugins/Lumin/MagicLeapDispatch/Source/MagicLeapDispatch.Build.cs b/Engine/Plugins/Lumin/MagicLeapDispatch/Source/MagicLeapDispatch.Build.cs
new file mode 100644
--- /dev/null
+++ b/Engine/Plugins/Lumin/MagicLeapDispatch/Source/MagicLeapDispatch.Build.cs
@@ -0,0 +1,31 @@
+// Copyright 1998-2019 Epic Games, Inc. All Rights Reserved.
+
+using System.IO;
+
+namespace UnrealBuildTool.Rules
+{
+	public class MagicLeapDispatch : ModuleRules
+	{
+		public MagicLeapDispatch( ReadOnlyTargetRules Target ) : base(Target)
+		{
+			PublicIncludePaths.Add(Path.Combine(ModuleDirectory, "Public"));
+
+			PublicDependencyModuleNames.AddRange(
+				new string[]
+				{
+					"Core",
+					"CoreUObject",
+					"Engine"
+				}
+			);
+
+			PrivateDependencyModuleNames.AddRange(
+				new string[]
+				{
+					"LuminRuntimeSettings",
+					"MLSDK"
+				}
+			);
+		}
+	}
+}
diff --git a/Engine/Plugins/Lumin/MagicLeapDispatch/Source/Private/MagicLeapDispatchPlugin.cpp b/Engine/Plugins/Lumin/MagicLeapDispatch/Source/Private/MagicLeapDispatchPlugin.cpp
new file mode 100644
--- /dev/null
+++ b/Engine/Plugins/Lumin/MagicLeapDispatch/Source/Private/MagicLeapDispatchPlugin.cpp
@@ -0,0 +1,113 @@
+// Copyright 1998-2019 Epic Games, Inc. All Rights Reserved.
+
+#include "IMagicLeapDispatchPlugin.h"
+#include "MagicLeapPluginUtil.h"
+#include "Lumin/CAPIShims/LuminAPIDispatch.h"
+
+#if PLATFORM_LUMIN
+#include "Lumin/LuminPlatformFile.h"
+#endif // PLATFORM_LUMIN
+
+DEFINE_LOG_CATEGORY_STATIC(LogMagicLeapDispatch, Log, All);
+
+#if WITH_MLSDK
+EMagicLeapDispatchResult MLToUEDispatchResult(MLResult Result)
+{
+	#define CASE_DISPATCH_RESULT(x) case MLDispatchResult_##x : return EMagicLeapDispatchResult::x;
+	#define CASE_ML_RESULT(x) case MLResult_##x : return EMagicLeapDispatchResult::x;
+
+	switch(Result)
+	{
+		CASE_ML_RESULT(Ok)
+		CASE_DISPATCH_RESULT(CannotStartApp)
+		CASE_DISPATCH_RESULT(InvalidPacket)
+		CASE_DISPATCH_RESULT(NoAppFound)
+		CASE_DISPATCH_RESULT(AppPickerDialogFailure)
+		CASE_ML_RESULT(AllocFailed)
+		CASE_ML_RESULT(InvalidParam)
+		CASE_ML_RESULT(UnspecifiedFailure)
+		CASE_ML_RESULT(NotImplemented)
+	}
+
+	return EMagicLeapDispatchResult::UnspecifiedFailure;
+}
+#endif // WITH_MLSDK
+
+class FMagicLeapDispatchPlugin : public IMagicLeapDispatchPlugin
+{
+public:
+#if PLATFORM_LUMIN
+	virtual EMagicLeapDispatchResult TryOpenApplication(const TArray<FLuminFileInfo>& DispatchFileList) override
+	{
+		EMagicLeapDispatchResult UEResult = EMagicLeapDispatchResult::NotImplemented;
+		IPlatformFile& PlatformFile = IPlatformFile::GetPlatformPhysical();
+		FLuminPlatformFile* LuminPlatformFile = static_cast<FLuminPlatformFile*>(&PlatformFile);
+
+		MLDispatchPacket* Packet = nullptr;
+		MLResult Result = MLDispatchAllocateEmptyPacket(&Packet);
+		if (Result == MLResult_Ok && Packet != nullptr)
+		{
+			if (DispatchFileList.Num() > 0)
+			{
+				Result = MLDispatchAllocateFileInfoList(Packet, DispatchFileList.Num());
+				UE_CLOG(Result != MLResult_Ok, LogMagicLeapDispatch, Error, TEXT("MLDispatchAllocateFileInfoList(%d) failed with error %s"), DispatchFileList.Num(), UTF8_TO_TCHAR(MLDispatchGetResultString(Result)));
+				if (Result == MLResult_Ok)
+				{
+					uint64_t AllocatedSize = 0;
+					Result = MLDispatchGetFileInfoListLength(Packet, &AllocatedSize);
+					UE_CLOG((AllocatedSize != DispatchFileList.Num()), LogMagicLeapDispatch, Error, TEXT("MLDispatchGetFileInfoListLength does not match expected length. %lld != %d"), AllocatedSize, DispatchFileList.Num());
+
+					uint64_t i = 0;
+					for (const FLuminFileInfo& DispatchFileInfo : DispatchFileList)
+					{
+						MLFileInfo* FileInfo = nullptr;
+						Result = MLDispatchGetFileInfoByIndex(Packet, i, &FileInfo);
+						UE_CLOG(Result != MLResult_Ok, LogMagicLeapDispatch, Error, TEXT("MLDispatchGetFileInfoByIndex(%lld) failed with error %s"), i, UTF8_TO_TCHAR(MLDispatchGetResultString(Result)));
+						++i;
+						if (Result == MLResult_Ok && FileInfo != nullptr)
+						{
+							if (DispatchFileInfo.MimeType.Len() > 0)
+							{
+								Result = MLFileInfoSetMimeType(FileInfo, TCHAR_TO_UTF8(*DispatchFileInfo.MimeType));
+								UE_CLOG(Result != MLResult_Ok, LogMagicLeapDispatch, Error, TEXT("MLFileInfoSetMimeType(%s) failed with error %s"), *DispatchFileInfo.MimeType, UTF8_TO_TCHAR(MLGetResultString(Result)));
+							}
+							if (DispatchFileInfo.FileName.Len() > 0)
+							{
+								UE_LOG(LogMagicLeapDispatch, Warning, TEXT("filename = %s"), *DispatchFileInfo.FileName);
+								Result = MLFileInfoSetFileName(FileInfo, TCHAR_TO_UTF8(*DispatchFileInfo.FileName));
+								UE_CLOG(Result != MLResult_Ok, LogMagicLeapDispatch, Error, TEXT("MLFileInfoSetFileName(%s) failed with error %s"), *DispatchFileInfo.FileName, UTF8_TO_TCHAR(MLGetResultString(Result)));
+							}
+							if (DispatchFileInfo.FileHandle != nullptr)
+							{
+								UE_LOG(LogMagicLeapDispatch, Warning, TEXT("setting file descriptor"));
+								LuminPlatformFile->SetMLFileInfoFD(DispatchFileInfo.FileHandle, FileInfo);
+							}
+
+							Result = MLDispatchAddFileInfo(Packet, FileInfo);
+							UE_CLOG(Result != MLResult_Ok, LogMagicLeapDispatch, Error, TEXT("MLDispatchAddFileInfo(%lld) failed with error %s"), i, UTF8_TO_TCHAR(MLDispatchGetResultString(Result)));
+						}
+					}
+				}
+			}
+
+			Result = MLDispatchTryOpenApplication(Packet);
+			UE_CLOG(Result != MLResult_Ok, LogMagicLeapDispatch, Error, TEXT("MLDispatchTryOpenApplication() failed with error %s"), UTF8_TO_TCHAR(MLDispatchGetResultString(Result)));
+			UEResult = MLToUEDispatchResult(Result);
+
+			Result = MLDispatchReleasePacket(&Packet, true, false);
+			UE_CLOG(Result != MLResult_Ok, LogMagicLeapDispatch, Error, TEXT("MLDispatchReleasePacket() failed with error %s"), UTF8_TO_TCHAR(MLDispatchGetResultString(Result)));
+		}
+		else
+		{
+			UEResult = MLToUEDispatchResult(Result);
+		}
+
+		return UEResult;
+	}
+#endif // PLATFORM_LUMIN
+
+private:
+	FMagicLeapAPISetup APISetup;
+};
+
+IMPLEMENT_MODULE(FMagicLeapDispatchPlugin, MagicLeapDispatch);
diff --git a/Engine/Plugins/Lumin/MagicLeapDispatch/Source/Public/IMagicLeapDispatchPlugin.h b/Engine/Plugins/Lumin/MagicLeapDispatch/Source/Public/IMagicLeapDispatchPlugin.h
new file mode 100644
--- /dev/null
+++ b/Engine/Plugins/Lumin/MagicLeapDispatch/Source/Public/IMagicLeapDispatchPlugin.h
@@ -0,0 +1,63 @@
+// Copyright 1998-2019 Epic Games, Inc. All Rights Reserved.
+
+#pragma once
+
+#include "Modules/ModuleManager.h"
+#include "UObject/ObjectMacros.h"
+#if PLATFORM_LUMIN
+#include "Lumin/LuminPlatformFile.h"
+#endif // PLATFORM_LUMIN
+
+UENUM(BlueprintType)
+enum class EMagicLeapDispatchResult : uint8
+{
+	Ok,
+	CannotStartApp,
+	InvalidPacket,
+	NoAppFound,
+	AppPickerDialogFailure,
+	AllocFailed,
+	InvalidParam,
+	UnspecifiedFailure,
+	NotImplemented
+};
+
+/**
+ * The public interface to this module.  In most cases, this interface is only public to sibling modules
+ * within this plugin.
+ */
+class MAGICLEAPDISPATCH_API IMagicLeapDispatchPlugin : public IModuleInterface
+{
+public:
+
+	/**
+	 * Singleton-like access to this module's interface.  This is just for convenience!
+	 * Beware of calling this during the shutdown phase, though.  Your module might have been unloaded already.
+	 *
+	 * @return Returns singleton instance, loading the module on demand if needed
+	 */
+	static inline IMagicLeapDispatchPlugin& Get()
+	{
+		return FModuleManager::LoadModuleChecked<IMagicLeapDispatchPlugin>("MagicLeapDispatch");
+	}
+
+	/**
+	 * Checks to see if this module is loaded and ready.  It is only valid to call Get() if IsAvailable() returns true.
+	 *
+	 * @return True if the module is loaded and ready to use
+	 */
+	static inline bool IsAvailable()
+	{
+		return FModuleManager::Get().IsModuleLoaded("MagicLeapDispatch");
+	}
+
+#if PLATFORM_LUMIN
+	/**
+	 * Attempt to launch another app on the Magic Leap device that can handle the mime types of file extensions of 
+	 * the files provided as arguments. These files will be sent to the target app as init args.
+	 * @param DispatchFileList A list of file info objects to be dispatched to another app. File info objects contain a file's name, mime type (optional) and it's IFileHandle pointer.
+	 * @return the dispatch result
+	 */
+	virtual EMagicLeapDispatchResult TryOpenApplication(const TArray<FLuminFileInfo>& DispatchFileList) = 0;
+#endif // PLATFORM_LUMIN
+};
diff --git a/Engine/Plugins/Lumin/MagicLeapLightEstimation/Config/BaseMagicLeapLightEstimation.ini b/Engine/Plugins/Lumin/MagicLeapLightEstimation/Config/BaseMagicLeapLightEstimation.ini
new file mode 100644
--- /dev/null
+++ b/Engine/Plugins/Lumin/MagicLeapLightEstimation/Config/BaseMagicLeapLightEstimation.ini
@@ -0,0 +1,5 @@
+[CoreRedirects]
++EnumRedirects=(OldName="/Script/MagicLeapLightEstimation.ELightEstimationCamera", NewName="/Script/MagicLeapLightEstimation.EMagicLeapLightEstimationCamera")
++StructRedirects=(OldName="/Script/MagicLeapLightEstimation.LightEstimationAmbientGlobalState", NewName="/Script/MagicLeapLightEstimation.MagicLeapLightEstimationAmbientGlobalState")
++StructRedirects=(OldName="/Script/MagicLeapLightEstimation.LightEstimationColorTemperatureState", NewName="/Script/MagicLeapLightEstimation.MagicLeapLightEstimationColorTemperatureState")
++ClassRedirects=(OldName="/Script/MagicLeapLightEstimation.LightingTrackingComponent", NewName="/Script/MagicLeapLightEstimation.MagicLeapLightingTrackingComponent")
diff --git a/Engine/Plugins/Lumin/MagicLeapLightEstimation/MagicLeapLightEstimation.uplugin b/Engine/Plugins/Lumin/MagicLeapLightEstimation/MagicLeapLightEstimation.uplugin
new file mode 100644
--- /dev/null
+++ b/Engine/Plugins/Lumin/MagicLeapLightEstimation/MagicLeapLightEstimation.uplugin
@@ -0,0 +1,42 @@
+{
+	"FileVersion" : 3,
+	"Version" : 1,
+	"VersionName": "1.0",
+	"FriendlyName" : "Magic Leap Light Estimatation",
+	"CreatedBy" : "Magic Leap, Inc.",
+	"CreatedByURL" : "http://magicleap.com",
+	"Description" : "Support for the Magic Leap Light Tracking API.",
+	"Category" : "Lumin",
+	"EnabledByDefault" : false,
+	"DocsURL" : "",
+	"MarketplaceURL" : "",
+	"SupportURL" : "",
+	"CanContainContent" : false,
+	"IsBetaVersion" : false,
+	"Installed" : false,
+	"SupportedTargetPlatforms": [
+		"Lumin",
+		"Win64",
+		"Mac"
+	],
+	"Modules" :
+	[
+		{
+			"Name" : "MagicLeapLightEstimation",
+			"Type" : "Runtime",
+			"WhitelistPlatforms" :
+			[
+				"Win64",
+				"Mac",
+				"Lumin"
+			]
+		}
+	],
+	"Plugins":
+	[
+		{
+			"Name": "MLSDK",
+			"Enabled": true
+		}
+	]
+}
diff --git a/Engine/Plugins/Lumin/MagicLeapLightEstimation/Source/MagicLeapLightEstimation/MagicLeapLightEstimation.Build.cs b/Engine/Plugins/Lumin/MagicLeapLightEstimation/Source/MagicLeapLightEstimation/MagicLeapLightEstimation.Build.cs
new file mode 100644
--- /dev/null
+++ b/Engine/Plugins/Lumin/MagicLeapLightEstimation/Source/MagicLeapLightEstimation/MagicLeapLightEstimation.Build.cs
@@ -0,0 +1,31 @@
+// Copyright 1998-2019 Epic Games, Inc. All Rights Reserved.
+
+using System.IO;
+
+namespace UnrealBuildTool.Rules
+{
+	public class MagicLeapLightEstimation : ModuleRules
+	{
+		public MagicLeapLightEstimation( ReadOnlyTargetRules Target ) : base(Target)
+		{
+			PublicIncludePaths.Add(Path.Combine(ModuleDirectory, "Public"));
+
+			PublicDependencyModuleNames.AddRange(
+				new string[]
+				{
+					"Core",
+					"CoreUObject",
+					"Engine"
+				}
+			);
+
+			PrivateDependencyModuleNames.AddRange(
+				new string[]
+				{
+					"LuminRuntimeSettings",
+					"MLSDK"
+				}
+			);
+		}
+	}
+}
diff --git a/Engine/Plugins/Lumin/MagicLeapLightEstimation/Source/MagicLeapLightEstimation/Private/LightingTrackingComponent.cpp b/Engine/Plugins/Lumin/MagicLeapLightEstimation/Source/MagicLeapLightEstimation/Private/LightingTrackingComponent.cpp
new file mode 100644
--- /dev/null
+++ b/Engine/Plugins/Lumin/MagicLeapLightEstimation/Source/MagicLeapLightEstimation/Private/LightingTrackingComponent.cpp
@@ -0,0 +1,140 @@
+// Copyright 1998-2019 Epic Games, Inc. All Rights Reserved.
+
+#include "LightingTrackingComponent.h"
+#include "IImageWrapper.h"
+#include "IImageWrapperModule.h"
+#include "Modules/ModuleManager.h"
+#include "Engine/TextureCube.h"
+#include "RenderUtils.h"
+#include "EngineUtils.h"
+#include "GameFramework/Actor.h"
+#include "Components/PostProcessComponent.h"
+#include "IMagicLeapLightEstimationPlugin.h"
+#include <limits>
+
+#define MAX_NITS 15.0f
+
+class FMagicLeapLightingTrackingImpl
+{
+public:
+	FMagicLeapLightingTrackingImpl(UMagicLeapLightingTrackingComponent* InOwner)
+		: Owner(InOwner)
+		, PostProcessor(nullptr)
+		, AmbientCubeMap(nullptr)
+		, LastAmbientIntensityTimeStamp(0)
+		, LastAmbientCubeMapTimeStamp(0)
+	{
+		IMagicLeapLightEstimationPlugin::Get().CreateTracker();
+
+		for (TActorIterator<AActor> ActorItr(Owner->GetWorld()); ActorItr; ++ActorItr)
+		{
+			PostProcessor = (*ActorItr)->FindComponentByClass<UPostProcessComponent>();
+			if (PostProcessor)
+			{
+				break;
+			}
+		}
+
+		if (!PostProcessor)
+		{
+			PostProcessor = NewObject<UPostProcessComponent>(Owner->GetOwner());
+			PostProcessor->RegisterComponent();
+		}
+
+		PostProcessor->bUnbound = true;
+		PostProcessor->Settings.bOverride_AmbientCubemapIntensity = Owner->UseGlobalAmbience;
+		PostProcessor->Settings.bOverride_WhiteTemp = Owner->UseColorTemp;
+	}
+
+	~FMagicLeapLightingTrackingImpl()
+	{
+		IMagicLeapLightEstimationPlugin::Get().DestroyTracker();
+	}
+
+	void RefreshGlobalAmbience()
+	{
+		PostProcessor->Settings.bOverride_AmbientCubemapIntensity = Owner->UseGlobalAmbience;
+
+		if (!Owner->UseGlobalAmbience)
+		{
+			return;
+		}
+
+		FMagicLeapLightEstimationAmbientGlobalState AmbientGlobalState;
+		bool Result = IMagicLeapLightEstimationPlugin::Get().GetAmbientGlobalState(AmbientGlobalState);
+		if (Result)
+		{
+			if (AmbientGlobalState.Timestamp > LastAmbientIntensityTimeStamp)
+			{
+				LastAmbientIntensityTimeStamp = AmbientGlobalState.Timestamp;
+				float LuminanceSum = 0.0f;
+				for (float GlobalLuminance : AmbientGlobalState.AmbientIntensityNits)
+				{
+					LuminanceSum += GlobalLuminance;
+				}
+				PostProcessor->Settings.AmbientCubemapIntensity = (float)(LuminanceSum / AmbientGlobalState.AmbientIntensityNits.Num()) / MAX_NITS;// (float)std::numeric_limits<uint16>::max();
+			}
+		}
+	}
+
+	void RefreshColorTemperature()
+	{
+		PostProcessor->Settings.bOverride_WhiteTemp = Owner->UseColorTemp;
+
+		if (!Owner->UseColorTemp)
+		{
+			return;
+		}
+
+		FMagicLeapLightEstimationColorTemperatureState ColorTemperatureState;
+		bool Result = IMagicLeapLightEstimationPlugin::Get().GetColorTemperatureState(ColorTemperatureState);
+		if (Result)
+		{
+			PostProcessor->Settings.WhiteTemp = ColorTemperatureState.ColorTemperatureKelvin;
+		}
+	}
+
+	void RefreshAmbientCubeMap()
+	{
+		
+	}
+
+	UMagicLeapLightingTrackingComponent* Owner;
+	UPostProcessComponent* PostProcessor;
+	UTextureCube* AmbientCubeMap;
+	FTimespan LastAmbientIntensityTimeStamp;
+	FTimespan LastAmbientCubeMapTimeStamp;
+};
+
+UMagicLeapLightingTrackingComponent::UMagicLeapLightingTrackingComponent()
+	: UseGlobalAmbience(false)
+	, UseColorTemp(false)
+	, Impl(nullptr)
+	//, UseDynamicAmbientCubeMap(false)
+{
+	PrimaryComponentTick.TickGroup = TG_PrePhysics;
+	PrimaryComponentTick.bStartWithTickEnabled = true;
+	PrimaryComponentTick.bCanEverTick = true;
+}
+
+void UMagicLeapLightingTrackingComponent::BeginPlay()
+{
+	UActorComponent::BeginPlay();
+	Impl = new FMagicLeapLightingTrackingImpl(this);
+}
+
+void UMagicLeapLightingTrackingComponent::EndPlay(const EEndPlayReason::Type EndPlayReason)
+{
+	UActorComponent::EndPlay(EndPlayReason);
+	delete Impl;
+	Impl = nullptr;
+}
+
+void UMagicLeapLightingTrackingComponent::TickComponent(float DeltaTime, enum ELevelTick TickType, FActorComponentTickFunction *ThisTickFunction)
+{
+	Impl->RefreshGlobalAmbience();
+
+	Impl->RefreshColorTemperature();
+
+	//Impl->RefreshAmbientCubeMap();
+}
\ No newline at end of file
diff --git a/Engine/Plugins/Lumin/MagicLeapLightEstimation/Source/MagicLeapLightEstimation/Private/MagicLeapLightEstimationFunctionLibrary.cpp b/Engine/Plugins/Lumin/MagicLeapLightEstimation/Source/MagicLeapLightEstimation/Private/MagicLeapLightEstimationFunctionLibrary.cpp
new file mode 100644
--- /dev/null
+++ b/Engine/Plugins/Lumin/MagicLeapLightEstimation/Source/MagicLeapLightEstimation/Private/MagicLeapLightEstimationFunctionLibrary.cpp
@@ -0,0 +1,29 @@
+// Copyright 1998-2019 Epic Games, Inc. All Rights Reserved.
+
+#include "MagicLeapLightEstimationFunctionLibrary.h"
+#include "MagicLeapLightEstimationModule.h"
+
+bool UMagicLeapLightEstimationFunctionLibrary::CreateTracker()
+{
+	return IMagicLeapLightEstimationPlugin::Get().CreateTracker();
+}
+
+void UMagicLeapLightEstimationFunctionLibrary::DestroyTracker()
+{
+	IMagicLeapLightEstimationPlugin::Get().DestroyTracker();
+}
+
+bool UMagicLeapLightEstimationFunctionLibrary::IsTrackerValid()
+{
+	return IMagicLeapLightEstimationPlugin::Get().IsTrackerValid();
+}
+
+bool UMagicLeapLightEstimationFunctionLibrary::GetAmbientGlobalState(FMagicLeapLightEstimationAmbientGlobalState& GlobalAmbientState)
+{
+	return IMagicLeapLightEstimationPlugin::Get().GetAmbientGlobalState(GlobalAmbientState);
+}
+
+bool UMagicLeapLightEstimationFunctionLibrary::GetColorTemperatureState(FMagicLeapLightEstimationColorTemperatureState& ColorTemperatureState)
+{
+	return IMagicLeapLightEstimationPlugin::Get().GetColorTemperatureState(ColorTemperatureState);
+}
diff --git a/Engine/Plugins/Lumin/MagicLeapLightEstimation/Source/MagicLeapLightEstimation/Private/MagicLeapLightEstimationModule.cpp b/Engine/Plugins/Lumin/MagicLeapLightEstimation/Source/MagicLeapLightEstimation/Private/MagicLeapLightEstimationModule.cpp
new file mode 100644
--- /dev/null
+++ b/Engine/Plugins/Lumin/MagicLeapLightEstimation/Source/MagicLeapLightEstimation/Private/MagicLeapLightEstimationModule.cpp
@@ -0,0 +1,103 @@
+// Copyright 1998-2019 Epic Games, Inc. All Rights Reserved.
+
+#include "MagicLeapLightEstimationModule.h"
+#include "Lumin/CAPIShims/LuminAPILightingTracking.h"
+
+DEFINE_LOG_CATEGORY(LogMagicLeapLightEstimation);
+
+void FMagicLeapLightEstimationModule::StartupModule()
+{
+#if WITH_MLSDK
+	Tracker = ML_INVALID_HANDLE;
+#endif // WITH_MLSDK
+}
+
+void FMagicLeapLightEstimationModule::ShutdownModule()
+{
+	DestroyTracker();
+}
+
+bool FMagicLeapLightEstimationModule::CreateTracker()
+{
+#if WITH_MLSDK
+	if (!MLHandleIsValid(Tracker))
+	{
+		MLResult Result = MLLightingTrackingCreate(&Tracker);
+		UE_CLOG(Result != MLResult_Ok, LogMagicLeapLightEstimation, Error, TEXT("MLLightingTrackingCreate failed with error %s"), UTF8_TO_TCHAR(MLGetResultString(Result)));
+	}
+
+	return MLHandleIsValid(Tracker);
+#else
+	return false;
+#endif // WITH_MLSDK
+}
+
+void FMagicLeapLightEstimationModule::DestroyTracker()
+{
+#if WITH_MLSDK
+	if (MLHandleIsValid(Tracker))
+	{
+		MLResult Result = MLLightingTrackingDestroy(Tracker);
+		UE_CLOG(Result != MLResult_Ok, LogMagicLeapLightEstimation, Error, TEXT("MLLightingTrackingDestroy failed with error %s"), UTF8_TO_TCHAR(MLGetResultString(Result)));
+	}
+#endif // WITH_MLSDK
+}
+
+bool FMagicLeapLightEstimationModule::IsTrackerValid() const
+{
+#if WITH_MLSDK
+	return MLHandleIsValid(Tracker);
+#else
+	return false;
+#endif // WITH_MLSDK
+}
+
+bool FMagicLeapLightEstimationModule::GetAmbientGlobalState(FMagicLeapLightEstimationAmbientGlobalState& GlobalAmbientState) const
+{
+#if WITH_MLSDK
+	MLResult Result = MLResult_InvalidParam;
+	if (MLHandleIsValid(Tracker))
+	{
+		MLLightingTrackingAmbientGlobalState State;
+		Result = MLLightingTrackingGetAmbientGlobalState(Tracker, &State);
+		if (Result == MLResult_Ok)
+		{
+			GlobalAmbientState.AmbientIntensityNits.Reset(MLLightingTrackingCamera_Count);
+			GlobalAmbientState.AmbientIntensityNits.AddUninitialized(MLLightingTrackingCamera_Count);
+			for (int32 i = 0; i < MLLightingTrackingCamera_Count; ++i)
+			{
+				GlobalAmbientState.AmbientIntensityNits[i] = State.als_global[i];				
+			}
+			GlobalAmbientState.Timestamp = FTimespan::FromMicroseconds(State.timestamp_ns / 1000.0f);
+		}
+	}
+
+	return Result == MLResult_Ok;
+#else
+	return false;
+#endif // WITH_MLSDK
+}
+
+bool FMagicLeapLightEstimationModule::GetColorTemperatureState(FMagicLeapLightEstimationColorTemperatureState& ColorTemperatureState) const
+{
+#if WITH_MLSDK
+	MLResult Result = MLResult_InvalidParam;
+	if (MLHandleIsValid(Tracker))
+	{
+		MLLightingTrackingColorTemperatureState State;
+		Result = MLLightingTrackingGetColorTemperatureState(Tracker, &State);
+		if (Result == MLResult_Ok)
+		{
+			ColorTemperatureState.ColorTemperatureKelvin = static_cast<float>(State.color_temp);
+			ColorTemperatureState.AmbientColor = FLinearColor(FVector(State.R_raw_pixel_avg, State.G_raw_pixel_avg, State.B_raw_pixel_avg));
+			ColorTemperatureState.Timestamp = FTimespan::FromMicroseconds(State.timestamp_ns / 1000.0f);
+		}
+	}
+
+	return Result == MLResult_Ok;
+#else
+	return false;
+#endif // WITH_MLSDK
+}
+
+IMPLEMENT_MODULE(FMagicLeapLightEstimationModule, MagicLeapLightEstimation);
diff --git a/Engine/Plugins/Lumin/MagicLeapLightEstimation/Source/MagicLeapLightEstimation/Private/MagicLeapLightEstimationModule.h b/Engine/Plugins/Lumin/MagicLeapLightEstimation/Source/MagicLeapLightEstimation/Private/MagicLeapLightEstimationModule.h
new file mode 100644
--- /dev/null
+++ b/Engine/Plugins/Lumin/MagicLeapLightEstimation/Source/MagicLeapLightEstimation/Private/MagicLeapLightEstimationModule.h
@@ -0,0 +1,26 @@
+// Copyright 1998-2019 Epic Games, Inc. All Rights Reserved.
+
+#pragma once
+
+#include "IMagicLeapLightEstimationPlugin.h"
+#include "Lumin/CAPIShims/LuminAPI.h"
+
+DECLARE_LOG_CATEGORY_EXTERN(LogMagicLeapLightEstimation, Verbose, All);
+
+class FMagicLeapLightEstimationModule : public IMagicLeapLightEstimationPlugin
+{
+public:
+	void StartupModule() override;
+	void ShutdownModule() override;
+
+	virtual bool CreateTracker() override;
+	virtual void DestroyTracker() override;
+	virtual bool IsTrackerValid() const override;
+	virtual bool GetAmbientGlobalState(FMagicLeapLightEstimationAmbientGlobalState& GlobalAmbientState) const override;
+	virtual bool GetColorTemperatureState(FMagicLeapLightEstimationColorTemperatureState& ColorTemperatureState) const override;
+
+private:
+#if WITH_MLSDK
+	MLHandle Tracker;
+#endif // WITH_MLSDK
+};
diff --git a/Engine/Plugins/Lumin/MagicLeapLightEstimation/Source/MagicLeapLightEstimation/Public/IMagicLeapLightEstimationPlugin.h b/Engine/Plugins/Lumin/MagicLeapLightEstimation/Source/MagicLeapLightEstimation/Public/IMagicLeapLightEstimationPlugin.h
new file mode 100644
--- /dev/null
+++ b/Engine/Plugins/Lumin/MagicLeapLightEstimation/Source/MagicLeapLightEstimation/Public/IMagicLeapLightEstimationPlugin.h
@@ -0,0 +1,67 @@
+// Copyright 1998-2019 Epic Games, Inc. All Rights Reserved.
+
+#pragma once
+
+#include "Modules/ModuleManager.h"
+#include "MagicLeapLightEstimationTypes.h"
+
+/**
+ * The public interface to this module.  In most cases, this interface is only public to sibling modules
+ * within this plugin.
+ */
+class MAGICLEAPLIGHTESTIMATION_API IMagicLeapLightEstimationPlugin : public IModuleInterface
+{
+public:
+
+	/**
+	 * Singleton-like access to this module's interface.  This is just for convenience!
+	 * Beware of calling this during the shutdown phase, though.  Your module might have been unloaded already.
+	 *
+	 * @return Returns singleton instance, loading the module on demand if needed
+	 */
+	static inline IMagicLeapLightEstimationPlugin& Get()
+	{
+		return FModuleManager::LoadModuleChecked<IMagicLeapLightEstimationPlugin>("MagicLeapLightEstimation");
+	}
+
+	/**
+	 * Checks to see if this module is loaded and ready.  It is only valid to call Get() if IsAvailable() returns true.
+	 *
+	 * @return True if the module is loaded and ready to use
+	 */
+	static inline bool IsAvailable()
+	{
+		return FModuleManager::Get().IsModuleLoaded("MagicLeapLightEstimation");
+	}
+
+	/**
+		Create a light estimation tracker.
+		@return true if tracker was successfully created, false otherwise
+	*/
+	virtual bool CreateTracker() = 0;
+
+	/** Destroy a light estimation tracker. */
+	virtual void DestroyTracker() = 0;
+
+	/**
+		Check if a light estimation tracker has already been created.
+		@return true if tracker already exists and is valid, false otherwise
+	*/
+	virtual bool IsTrackerValid() const = 0;
+
+	/**
+		Gets information about the ambient light sensor global state.
+		@note Capturing images or video will stop the lighting information update, causing the retrieved data to be stale (old timestamps).
+		@param GlobalAmbientState Output param containing the information about the global lighting state (ambient intensity). Valid only if return value of function is true.
+		@return true if the global ambient state was succesfully retrieved.
+	*/
+	virtual bool GetAmbientGlobalState(FMagicLeapLightEstimationAmbientGlobalState& GlobalAmbientState) const = 0;
+
+	/**
+		Gets information about the color temperature state.
+		@note Capturing images or video will stop the lighting information update, causing the retrieved data to be stale (old timestamps).
+		@param ColorTemperatureState Output param containing the information about the color temperature. Valid only if return value of function is true.
+		@return true if the color temperature state was succesfully retrieved.
+	*/
+	virtual bool GetColorTemperatureState(FMagicLeapLightEstimationColorTemperatureState& ColorTemperatureState) const = 0;
+};
diff --git a/Engine/Plugins/Lumin/MagicLeapLightEstimation/Source/MagicLeapLightEstimation/Public/LightingTrackingComponent.h b/Engine/Plugins/Lumin/MagicLeapLightEstimation/Source/MagicLeapLightEstimation/Public/LightingTrackingComponent.h
new file mode 100644
--- /dev/null
+++ b/Engine/Plugins/Lumin/MagicLeapLightEstimation/Source/MagicLeapLightEstimation/Public/LightingTrackingComponent.h
@@ -0,0 +1,41 @@
+// Copyright 1998-2019 Epic Games, Inc. All Rights Reserved.
+
+#pragma once
+
+#include "Components/ActorComponent.h"
+#include "Engine/Engine.h"
+#include "LightingTrackingComponent.generated.h"
+
+/**
+The LightingTrackingComponent wraps the Magic Leap lighting tracking API.
+This api provides lumosity data from the camera that can be used to shade objects in a more realistic
+manner (via the post processor).
+*/
+UCLASS(ClassGroup = MagicLeap, BlueprintType, Blueprintable, EditInlineNew, meta = (BlueprintSpawnableComponent))
+class MAGICLEAPLIGHTESTIMATION_API UMagicLeapLightingTrackingComponent : public UActorComponent
+{
+	GENERATED_BODY()
+
+public:
+	UMagicLeapLightingTrackingComponent();
+
+	/** Intializes the lighting tracking api. If a post processing component cannot be found in the scene, one will be created. */
+	void BeginPlay() override;
+	/** Cleans up the lighting tracking api. */
+	void EndPlay(const EEndPlayReason::Type EndPlayReason) override;
+	/** Polls for data from the camera array and processes it based on the active modes (UseGlobalAmbience, UseColorTemp, ...). */
+	virtual void TickComponent(float DeltaTime, enum ELevelTick TickType, FActorComponentTickFunction *ThisTickFunction) override;
+
+	/** Set to true if you want the global ambience value from the cameras to be used in post processing. */
+	UPROPERTY(EditAnywhere, BlueprintReadWrite, Category = "LightingTracking|MagicLeap")
+	bool UseGlobalAmbience;
+	/** Set to true if you want the color temperature value from the cameras to be used in post processing. */
+	UPROPERTY(EditAnywhere, BlueprintReadWrite, Category = "LightingTracking|MagicLeap")
+	bool UseColorTemp;
+	/** Set to true if you want the ambient cube map to be dynamically updated from the cameras' data. */
+	//UPROPERTY(EditAnywhere, BlueprintReadWrite, Category = "LightingTracking|MagicLeap")
+	//bool UseDynamicAmbientCubeMap;
+
+private:
+	class FMagicLeapLightingTrackingImpl* Impl;
+};
diff --git a/Engine/Plugins/Lumin/MagicLeapLightEstimation/Source/MagicLeapLightEstimation/Public/MagicLeapLightEstimationFunctionLibrary.h b/Engine/Plugins/Lumin/MagicLeapLightEstimation/Source/MagicLeapLightEstimation/Public/MagicLeapLightEstimationFunctionLibrary.h
new file mode 100644
--- /dev/null
+++ b/Engine/Plugins/Lumin/MagicLeapLightEstimation/Source/MagicLeapLightEstimation/Public/MagicLeapLightEstimationFunctionLibrary.h
@@ -0,0 +1,50 @@
+// Copyright 1998-2019 Epic Games, Inc. All Rights Reserved.
+#pragma once
+
+#include "Engine/Engine.h"
+#include "Kismet/BlueprintFunctionLibrary.h"
+#include "MagicLeapLightEstimationTypes.h"
+#include "MagicLeapLightEstimationFunctionLibrary.generated.h"
+
+UCLASS(ClassGroup = MagicLeap)
+class MAGICLEAPLIGHTESTIMATION_API UMagicLeapLightEstimationFunctionLibrary : public UBlueprintFunctionLibrary
+{
+	GENERATED_BODY()
+
+public:
+	/**
+		Create a light estimation tracker.
+		@return true if tracker was successfully created, false otherwise
+	*/
+	UFUNCTION(BlueprintCallable, Category = "LightEstimation Function Library | MagicLeap")
+	static bool CreateTracker();
+
+	/** Destroy a light estimation tracker. */
+	UFUNCTION(BlueprintCallable, Category = "LightEstimation Function Library | MagicLeap")
+	static void DestroyTracker();
+
+	/**
+		Check if a light estimation tracker has already been created.
+		@return true if tracker already exists and is valid, false otherwise
+	*/
+	UFUNCTION(BlueprintCallable, BlueprintPure, Category = "LightEstimation Function Library | MagicLeap")
+	static bool IsTrackerValid();
+
+	/**
+		Gets information about the ambient light sensor global state.
+		@note Capturing images or video will stop the lighting information update, causing the retrieved data to be stale (old timestamps).
+		@param GlobalAmbientState Output param containing the information about the global lighting state (ambient intensity). Valid only if return value of function is true.
+		@return true if the global ambient state was succesfully retrieved.
+	*/
+	UFUNCTION(BlueprintCallable, Category = "LightEstimation Function Library | MagicLeap")
+	static bool GetAmbientGlobalState(FMagicLeapLightEstimationAmbientGlobalState& GlobalAmbientState);
+
+	/**
+		Gets information about the color temperature state.
+		@note Capturing images or video will stop the lighting information update, causing the retrieved data to be stale (old timestamps).
+		@param ColorTemperatureState Output param containing the information about the color temperature. Valid only if return value of function is true.
+		@return true if the color temperature state was succesfully retrieved.
+	*/
+	UFUNCTION(BlueprintCallable, Category = "LightEstimation Function Library | MagicLeap")
+	static bool GetColorTemperatureState(FMagicLeapLightEstimationColorTemperatureState& ColorTemperatureState);
+};
diff --git a/Engine/Plugins/Lumin/MagicLeapLightEstimation/Source/MagicLeapLightEstimation/Public/MagicLeapLightEstimationTypes.h b/Engine/Plugins/Lumin/MagicLeapLightEstimation/Source/MagicLeapLightEstimation/Public/MagicLeapLightEstimationTypes.h
new file mode 100644
--- /dev/null
+++ b/Engine/Plugins/Lumin/MagicLeapLightEstimation/Source/MagicLeapLightEstimation/Public/MagicLeapLightEstimationTypes.h
@@ -0,0 +1,59 @@
+// Copyright 1998-2019 Epic Games, Inc. All Rights Reserved.
+#pragma once
+
+#include "Engine/Engine.h"
+#include "MagicLeapLightEstimationTypes.generated.h"
+
+/** Which camera the information is related to. */
+UENUM(BlueprintType)
+enum class EMagicLeapLightEstimationCamera : uint8
+{
+	/** Left. */
+	Left,
+	/** Right. */
+	Right,
+	/** Far left. */
+	FarLeft,
+	/** Far right. */
+	FarRight
+};
+
+/**
+	Information about the ambient light sensor global state.
+	Capturing images or video will stop the lighting information update,
+	causing the retrieved data to be stale (old timestamps).
+*/
+USTRUCT(BlueprintType)
+struct FMagicLeapLightEstimationAmbientGlobalState
+{
+	GENERATED_BODY()
+
+public:
+	/* Array stores values for each world camera, ordered left, right, far left, far right. Luminance estimate is in nits (cd/m^2). */
+	UPROPERTY(BlueprintReadOnly, Category = "Light Estimation|MagicLeap")
+	TArray<float> AmbientIntensityNits;
+
+	UPROPERTY(BlueprintReadOnly, Category = "Light Estimation|MagicLeap")
+	FTimespan Timestamp;
+};
+
+/**
+	Information about the color temperature state.
+	Capturing images or video will stop the lighting information update,
+	causing the retrieved data to be stale (old timestamps).
+*/
+USTRUCT(BlueprintType)
+struct FMagicLeapLightEstimationColorTemperatureState
+{
+	GENERATED_BODY()
+
+public:
+	UPROPERTY(BlueprintReadOnly, Category = "Light Estimation|MagicLeap")
+	float ColorTemperatureKelvin;
+
+	UPROPERTY(BlueprintReadOnly, Category = "Light Estimation|MagicLeap")
+	FLinearColor AmbientColor;
+
+	UPROPERTY(BlueprintReadOnly, Category = "Light Estimation|MagicLeap")
+	FTimespan Timestamp;
+};
diff --git a/Engine/Plugins/Lumin/MagicLeapLocation/MagicLeapLocation.uplugin b/Engine/Plugins/Lumin/MagicLeapLocation/MagicLeapLocation.uplugin
new file mode 100644
--- /dev/null
+++ b/Engine/Plugins/Lumin/MagicLeapLocation/MagicLeapLocation.uplugin
@@ -0,0 +1,40 @@
+{
+	"FileVersion": 3,
+	"Version": 1,
+	"VersionName": "1.0",
+	"FriendlyName": "Magic Leap Location",
+	"CreatedBy": "Magic Leap, Inc.",
+	"CreatedByURL": "http://magicleap.com",
+	"Description": "Support for the Magic Leap Location API.",
+	"Category": "Lumin",
+	"EnabledByDefault": false,
+	"DocsURL": "",
+	"MarketplaceURL": "",
+	"SupportURL": "",
+	"CanContainContent": false,
+	"IsBetaVersion": false,
+	"Installed": false,
+	"SupportedTargetPlatforms": [
+		"Win64",
+		"Mac",
+		"Lumin"
+	],
+	"Modules": [
+		{
+			"Name": "MagicLeapLocation",
+			"Type": "Runtime",
+			"LoadingPhase": "PostConfigInit",
+			"WhitelistPlatforms": [
+				"Win64",
+				"Mac",
+				"Lumin"
+			]
+		}
+	],
+	"Plugins": [
+		{
+			"Name": "MagicLeap",
+			"Enabled": true
+		}
+	]
+}
diff --git a/Engine/Plugins/Lumin/MagicLeapLocation/Source/MagicLeapLocation.Build.cs b/Engine/Plugins/Lumin/MagicLeapLocation/Source/MagicLeapLocation.Build.cs
new file mode 100644
--- /dev/null
+++ b/Engine/Plugins/Lumin/MagicLeapLocation/Source/MagicLeapLocation.Build.cs
@@ -0,0 +1,32 @@
+// Copyright 1998-2019 Epic Games, Inc. All Rights Reserved.
+
+using System.IO;
+
+namespace UnrealBuildTool.Rules
+{
+	public class MagicLeapLocation : ModuleRules
+	{
+		public MagicLeapLocation( ReadOnlyTargetRules Target ) : base(Target)
+		{
+			PublicIncludePaths.Add(Path.Combine(ModuleDirectory, "Public"));
+
+			PublicDependencyModuleNames.AddRange(
+				new string[]
+				{
+					"Core",
+					"CoreUObject",
+					"Engine"
+				}
+			);
+
+			PrivateDependencyModuleNames.AddRange(
+				new string[]
+				{
+					"LuminRuntimeSettings",
+					"MLSDK",
+					"MagicLeap",
+				}
+			);
+		}
+	}
+}
