Fix static analysis warnings on email and phonenumber copy code (Dan O'Connor and I trawled through the code and agreed the buffer overrun in this case won't happen).
#rnx

diff --git a/Engine/Plugins/Lumin/MagicLeapContacts/Source/Private/MagicLeapContactsPlugin.cpp b/Engine/Plugins/Lumin/MagicLeapContacts/Source/Private/MagicLeapContactsPlugin.cpp
--- a/Engine/Plugins/Lumin/MagicLeapContacts/Source/Private/MagicLeapContactsPlugin.cpp
+++ b/Engine/Plugins/Lumin/MagicLeapContacts/Source/Private/MagicLeapContactsPlugin.cpp
@@ -388,6 +388,9 @@ MLContactsSearchField FMagicLeapContactsPlugin::UESearchFieldToMLSearchField(EMa
 	return SearchField;
 }
 
+#pragma warning(disable:6386) /* Buffer overrun while writing to '<array>': the writable size is '<array size>' bytes, but '<write size>' bytes might be written */
+							  /* After careful code analysis, we've determine this warning is incorrect in this case, do disabling at function scope. */
+
 void FMagicLeapContactsPlugin::UEToMLContact(const FMagicLeapContact& InUEContact, MLContactsContact& OutMLContact)
 {
 	MLContactsContactInit(&OutMLContact);
@@ -430,6 +433,7 @@ void FMagicLeapContactsPlugin::UEToMLContact(const FMagicLeapContact& InUEContac
 		MLEmailAddresses[iEmailAddress] = MLEmailAddress;
 	}
 }
+#pragma warning(enable:6386)	/* Re-enable warning outside scope of function. */
 
 void FMagicLeapContactsPlugin::DestroyMLContact(MLContactsContact& OutMLContact)
 {
