Fix compile errors due to CL 9577777.

#rnx

diff --git a/Engine/Plugins/Lumin/MagicLeapContacts/Source/Private/MagicLeapContactsPlugin.cpp b/Engine/Plugins/Lumin/MagicLeapContacts/Source/Private/MagicLeapContactsPlugin.cpp
--- a/Engine/Plugins/Lumin/MagicLeapContacts/Source/Private/MagicLeapContactsPlugin.cpp
+++ b/Engine/Plugins/Lumin/MagicLeapContacts/Source/Private/MagicLeapContactsPlugin.cpp
@@ -388,8 +388,11 @@ MLContactsSearchField FMagicLeapContactsPlugin::UESearchFieldToMLSearchField(EMa
 	return SearchField;
 }
 
+#ifdef _MSC_VER
+#pragma warning(push)
 #pragma warning(disable:6386) /* Buffer overrun while writing to '<array>': the writable size is '<array size>' bytes, but '<write size>' bytes might be written */
 							  /* After careful code analysis, we've determine this warning is incorrect in this case, do disabling at function scope. */
+#endif
 
 void FMagicLeapContactsPlugin::UEToMLContact(const FMagicLeapContact& InUEContact, MLContactsContact& OutMLContact)
 {
@@ -433,7 +436,10 @@ void FMagicLeapContactsPlugin::UEToMLContact(const FMagicLeapContact& InUEContac
 		MLEmailAddresses[iEmailAddress] = MLEmailAddress;
 	}
 }
-#pragma warning(enable:6386)	/* Re-enable warning outside scope of function. */
+
+#ifdef _MSC_VER
+#pragma warning(pop)
+#endif
 
 void FMagicLeapContactsPlugin::DestroyMLContact(MLContactsContact& OutMLContact)
 {
